Fix neural post processing compile error without editor.

diff --git a/Engine/Plugins/Experimental/NeuralRendering/Source/NeuralPostProcessing/Private/NeuralPostProcessModelInstance.cpp b/Engine/Plugins/Experimental/NeuralRendering/Source/NeuralPostProcessing/Private/NeuralPostProcessModelInstance.cpp
--- a/Engine/Plugins/Experimental/NeuralRendering/Source/NeuralPostProcessing/Private/NeuralPostProcessModelInstance.cpp
+++ b/Engine/Plugins/Experimental/NeuralRendering/Source/NeuralPostProcessing/Private/NeuralPostProcessModelInstance.cpp
@@ -13,20 +13,26 @@ TSharedPtr<UE::NNE::IModelInstanceRDG> CreateNNEModelInstance(UNNEModelData* NNE
 	TWeakInterfacePtr<INNERuntimeRDG> RuntimeRDG = UE::NNE::GetRuntime<INNERuntimeRDG>(RuntimeName);
 	if (!Runtime.IsValid())
 	{
+#if WITH_EDITOR
 		UE_LOG(LogNeuralPostProcessing, Error, TEXT("Can't get %s runtime."), *RuntimeName);
+#endif
 		return nullptr;
 	}
 
 	if (!RuntimeRDG.IsValid())
 	{
+#if WITH_EDITOR
 		UE_LOG(LogNeuralPostProcessing, Error, TEXT("No RDG runtime '%s' found"), *RuntimeName);
+#endif
 		return nullptr;
 	}
 
 	TSharedPtr<UE::NNE::IModelRDG> ModelRDG = RuntimeRDG->CreateModelRDG(NNEModelData);
 	if (!ModelRDG.IsValid())
 	{
+#if WITH_EDITOR
 		UE_LOG(LogNeuralPostProcessing, Error, TEXT("CreateModelRDG failed for Runtime = %s"), *RuntimeName);
+#endif
 		return nullptr;
 	}
 
@@ -41,20 +47,26 @@ TSharedPtr<UE::NNE::IModelInstanceCPU> CreateNNECpuModelInstance(UNNEModelData*
 	TWeakInterfacePtr<INNERuntimeCPU> RuntimeCPU = UE::NNE::GetRuntime<INNERuntimeCPU>(RuntimeName);
 	if (!Runtime.IsValid())
 	{
+#if WITH_EDITOR
 		UE_LOG(LogNeuralPostProcessing, Error, TEXT("Can't get %s runtime."), *RuntimeName);
+#endif
 		return nullptr;
 	}
 
 	if (!RuntimeCPU.IsValid())
 	{
+#if WITH_EDITOR
 		UE_LOG(LogNeuralPostProcessing, Error, TEXT("No CPU runtime '%s' found"), *RuntimeName);
+#endif
 		return nullptr;
 	}
 
 	TSharedPtr<UE::NNE::IModelCPU> ModelCPU = RuntimeCPU->CreateModelCPU(NNEModelData);
 	if (!ModelCPU.IsValid())
 	{
+#if WITH_EDITOR
 		UE_LOG(LogNeuralPostProcessing, Error, TEXT("CreateModelCPU failed for Runtime = %s"), *RuntimeName);
+#endif
 		return nullptr;
 	}
 
@@ -179,7 +191,9 @@ bool UNeuralPostProcessModelInstance::ModifyInputShape(int Dim, int Size)
 
 	if (NewResolvedInputShape[Dim] != Size)
 	{
+#if WITH_EDITOR
 		UE_LOG(LogNeuralPostProcessing, Error, TEXT("Cannot set dimension %d to %d. It is dynamic. revert back to %d"),Dim, Size, NewResolvedInputShape[Dim]);
+#endif
 		return false;
 	}
 
@@ -192,7 +206,9 @@ bool UNeuralPostProcessModelInstance::ModifyInputShape(int Dim, int Size)
 	TConstArrayView<UE::NNE::FTensorShape> ResolvedOutputTensorShapes = ModelInstanceRDG->GetOutputTensorShapes();
 	if (ResolvedOutputTensorShapes.Num() < 1)
 	{
+#if WITH_EDITOR
 		UE_LOG(LogNeuralPostProcessing, Error, TEXT("Cannot set dimension %d to %d due to NNE eror. Revert back to %d"), Dim, Size, CurrentResolvedInputTensorShape.GetData()[Dim]);
+#endif
 		ModelInstanceRDG->SetInputTensorShapes({ CurrentResolvedInputTensorShape });
 		return false;
 	}
@@ -282,6 +298,8 @@ void UNeuralPostProcessModelInstance::CreateDefaultNNEModel(UNNEModelData* NNEMo
 	}
 	else
 	{
+#if WITH_EDITOR
 		UE_LOG(LogNeuralPostProcessing, Error, TEXT("Failed to create NNE RDG Model."));
+#endif
 	}
 }
diff --git a/Engine/Plugins/Experimental/NeuralRendering/Source/NeuralPostProcessing/Private/NeuralPostProcessing.cpp b/Engine/Plugins/Experimental/NeuralRendering/Source/NeuralPostProcessing/Private/NeuralPostProcessing.cpp
--- a/Engine/Plugins/Experimental/NeuralRendering/Source/NeuralPostProcessing/Private/NeuralPostProcessing.cpp
+++ b/Engine/Plugins/Experimental/NeuralRendering/Source/NeuralPostProcessing/Private/NeuralPostProcessing.cpp
@@ -105,7 +105,9 @@ public:
 
 		if (!IsValid(RawNNEModelData))
 		{
+#if WITH_EDITOR
 			UE_LOG(LogNeuralPostProcessing, Error, TEXT("NNEModelData is invalid at Slot %d."), AllocationId);
+#endif
 			return;
 		}
 
@@ -451,7 +453,9 @@ void FNeuralPostProcessingModule::StartupModule()
 	}
 	else
 	{
+#if WITH_EDITOR
 		UE_LOG(LogNeuralPostProcessing, Error, TEXT("Shaders directory not added. Failed to find NeuralPostProcessing plugin"));
+#endif
 	}
 
 	GNeuralProfileManager.Reset(new FNeuralProfileManager());
