Fix incremental UnrealEditor Linux Error for neural post processing

diff --git a/Engine/Plugins/Experimental/NeuralRendering/Source/NeuralPostProcessing/Private/NeuralPostProcessModelInstance.cpp b/Engine/Plugins/Experimental/NeuralRendering/Source/NeuralPostProcessing/Private/NeuralPostProcessModelInstance.cpp
--- a/Engine/Plugins/Experimental/NeuralRendering/Source/NeuralPostProcessing/Private/NeuralPostProcessModelInstance.cpp
+++ b/Engine/Plugins/Experimental/NeuralRendering/Source/NeuralPostProcessing/Private/NeuralPostProcessModelInstance.cpp
@@ -187,7 +187,7 @@ bool UNeuralPostProcessModelInstance::ModifyInputShape(int Dim, int Size)
 	uint32 Height = (Dim == 2) && (InputShapeTemplate.GetData()[2] < 0) ? Size : CurrentResolvedInputTensorShape.GetData()[2];
 	uint32 Width = (Dim == 3) && (InputShapeTemplate.GetData()[3] < 0) ? Size : CurrentResolvedInputTensorShape.GetData()[3];
 	
-	TConstArrayView<uint32> NewResolvedInputShape = {Batch,Channel,Height,Width};
+	TArray<uint32> NewResolvedInputShape = {Batch,Channel,Height,Width};
 
 	if (NewResolvedInputShape[Dim] != Size)
 	{
