Fix NeuralPostProcessing references the Private Directory of the Renderer Module (PostProcess/PostProcessing.h).
* Move the dependent INeuralPostProcessInterface interface to internal folder
* Mark the plugin as engine module.

diff --git a/Engine/Plugins/Experimental/NeuralRendering/Source/NeuralPostProcessing/NeuralPostProcessing.Build.cs b/Engine/Plugins/Experimental/NeuralRendering/Source/NeuralPostProcessing/NeuralPostProcessing.Build.cs
--- a/Engine/Plugins/Experimental/NeuralRendering/Source/NeuralPostProcessing/NeuralPostProcessing.Build.cs
+++ b/Engine/Plugins/Experimental/NeuralRendering/Source/NeuralPostProcessing/NeuralPostProcessing.Build.cs
@@ -6,6 +6,7 @@ public class NeuralPostProcessing : ModuleRules
 {
 	public NeuralPostProcessing(ReadOnlyTargetRules Target) : base(Target)
 	{
+		bTreatAsEngineModule = true;
 		PCHUsage = ModuleRules.PCHUsageMode.UseExplicitOrSharedPCHs;
 		
 		PublicIncludePaths.AddRange(
@@ -18,7 +19,6 @@ public class NeuralPostProcessing : ModuleRules
 		PrivateIncludePaths.AddRange(
 			new string[] {
 				// ... add other private include paths required here ...
-				System.IO.Path.Combine(EngineDirectory, "Source/Runtime/Renderer/Private")
 			}
 			);
 			
diff --git a/Engine/Plugins/Experimental/NeuralRendering/Source/NeuralPostProcessing/Private/NeuralPostProcessing.cpp b/Engine/Plugins/Experimental/NeuralRendering/Source/NeuralPostProcessing/Private/NeuralPostProcessing.cpp
--- a/Engine/Plugins/Experimental/NeuralRendering/Source/NeuralPostProcessing/Private/NeuralPostProcessing.cpp
+++ b/Engine/Plugins/Experimental/NeuralRendering/Source/NeuralPostProcessing/Private/NeuralPostProcessing.cpp
@@ -8,8 +8,7 @@
 #include "Engine/NeuralProfile.h"
 #include "NeuralPostProcessModelInstance.h"
 #include "NeuralPostProcessingCS.h"
-#include "PostProcess/PostProcessing.h"
-#include "PostProcess/NeuralPostProcess.h"
+#include "PostProcess/NeuralPostProcessInterface.h"
 #include "PixelShaderUtils.h"
 #include "RenderGraphEvent.h"
 
@@ -275,7 +274,6 @@ static void AllocateInputBuffer_RenderingThread(
 
 static void	ApplyNeuralNetworks_RenderingThread(
 	FRDGBuilder& GraphBuilder,
-	const FViewInfo& View,
 	int32 ProfileId,
 	FRDGTextureRef NeuralTexture,
 	FIntRect ViewRect,
@@ -319,7 +317,7 @@ static void	ApplyNeuralNetworks_RenderingThread(
 			PassParameters->NetworkTextureSize = NeuralNetworkInputSize;
 			PassParameters->SourceType = GraphBuilder.CreateSRV(InputSourceType, EPixelFormat::PF_R32_UINT);
 			PassParameters->RWIndirectDispatchArgsBuffer = GraphBuilder.CreateUAV(IndirectDispatchBuffer, EPixelFormat::PF_R32_UINT);
-			TShaderMapRef<ARGSETUPSHADER> ComputeShader(View.ShaderMap);
+			TShaderMapRef<ARGSETUPSHADER> ComputeShader(GlobalShaderMap);
 			FComputeShaderUtils::AddPass(GraphBuilder, FRDGEventName(TEXT("NeuralPostProcessing::BuildIndirectArgs(Dispatch)")), ComputeShader, PassParameters, FIntVector(1,1,1));
 		}
 
@@ -412,13 +410,12 @@ static void	ApplyNeuralNetworks_RenderingThread(
 class FNeuralPostProcess : public INeuralPostProcessInterface
 {
 public:
-	virtual void Apply(FRDGBuilder& GraphBuilder, const FViewInfo& View, int32 NeuralProfileId,
+	virtual void Apply(FRDGBuilder& GraphBuilder, int32 NeuralProfileId,
 		FRDGTexture* NeuralTexture, FIntRect ViewRect, FRDGBufferRef InputSourceType,
 		FRDGBufferRef& OutputNeuralBuffer, FVector4f& BufferDimension) override
 	{
 		ApplyNeuralNetworks_RenderingThread(
 			GraphBuilder,
-			View,
 			NeuralProfileId,
 			NeuralTexture,
 			ViewRect,
diff --git a/Engine/Source/Runtime/Renderer/Internal/PostProcess/NeuralPostProcessInterface.h b/Engine/Source/Runtime/Renderer/Internal/PostProcess/NeuralPostProcessInterface.h
new file mode 100644
--- /dev/null
+++ b/Engine/Source/Runtime/Renderer/Internal/PostProcess/NeuralPostProcessInterface.h
@@ -0,0 +1,21 @@
+// Copyright Epic Games, Inc. All Rights Reserved.
+
+#pragma once
+
+#include "ScreenPass.h"
+
+class INeuralPostProcessInterface
+{
+public:
+
+	virtual ~INeuralPostProcessInterface() {}
+
+	virtual void Apply(FRDGBuilder& GraphBuilder, int32 NeuralProfileId,
+		FRDGTexture* NeuralTexture, FIntRect ViewRect, FRDGBufferRef InputSourceType,
+		FRDGBufferRef& OutputNeuralBuffer, FVector4f& BufferDimension) = 0;
+
+	virtual void AllocateBuffer(FRDGBuilder& GraphBuilder, const FScreenPassTextureViewport& Viewport,
+		int32 NeuralProfileId, FRDGBufferRef& InputNeuralBuffer, FVector4f& InputBufferDimension) = 0;
+};
+
+extern RENDERER_API TUniquePtr<INeuralPostProcessInterface> GNeuralPostProcess;
\ No newline at end of file
diff --git a/Engine/Source/Runtime/Renderer/Private/PostProcess/NeuralPostProcess.cpp b/Engine/Source/Runtime/Renderer/Private/PostProcess/NeuralPostProcess.cpp
--- a/Engine/Source/Runtime/Renderer/Private/PostProcess/NeuralPostProcess.cpp
+++ b/Engine/Source/Runtime/Renderer/Private/PostProcess/NeuralPostProcess.cpp
@@ -148,7 +148,6 @@ void ApplyNeuralPostProcess(
 	{
 		GNeuralPostProcess->Apply(
 			GraphBuilder,
-			View,
 			NeuralPostProcessResource.NeuralProfileId,
 			NeuralPostProcessResource.Texture,
 			Rect,
diff --git a/Engine/Source/Runtime/Renderer/Private/PostProcess/NeuralPostProcess.h b/Engine/Source/Runtime/Renderer/Private/PostProcess/NeuralPostProcess.h
--- a/Engine/Source/Runtime/Renderer/Private/PostProcess/NeuralPostProcess.h
+++ b/Engine/Source/Runtime/Renderer/Private/PostProcess/NeuralPostProcess.h
@@ -5,6 +5,7 @@
 #include "ScreenPass.h"
 #include "OverridePassSequence.h"
 #include "PostProcess/PostProcessMaterialInputs.h"
+#include "PostProcess/NeuralPostProcessInterface.h"
 
 BEGIN_SHADER_PARAMETER_STRUCT(FNeuralPostProcessShaderParameters, )
 	SHADER_PARAMETER_RDG_TEXTURE_UAV(RWTexture2D<float4>, RWNeuralTexture)
@@ -64,19 +65,3 @@ void ApplyNeuralPostProcess(
 bool IsNeuralPostProcessEnabled();
 bool ShouldApplyNeuralPostProcessForMaterial(const class FMaterial* Material);
 bool IsNeuralPostProcessShaderParameterUsed(FNeuralPostProcessShaderParameters& NeuralPostProcessShaderParameters);
-
-class INeuralPostProcessInterface
-{
-public:
-
-	virtual ~INeuralPostProcessInterface() {}
-
-	virtual void Apply(FRDGBuilder& GraphBuilder, const FViewInfo& View, int32 NeuralProfileId,
-		FRDGTexture* NeuralTexture, FIntRect ViewRect, FRDGBufferRef InputSourceType,
-		FRDGBufferRef& OutputNeuralBuffer, FVector4f& BufferDimension) = 0;
-
-	virtual void AllocateBuffer(FRDGBuilder& GraphBuilder, const FScreenPassTextureViewport& Viewport,
-		int32 NeuralProfileId, FRDGBufferRef& InputNeuralBuffer, FVector4f& InputBufferDimension) = 0;
-};
-
-extern RENDERER_API TUniquePtr<INeuralPostProcessInterface> GNeuralPostProcess;
