Code clean up for meta and build config of neural post processing. meta edit conditon for vectors (e.g., V[0]==-1 for FVector4 ) is not supported at this moment. It is now causing errors in the editor session log. Will work on it later.

diff --git a/Engine/Plugins/Experimental/NeuralRendering/Source/NeuralPostProcessing/NeuralPostProcessing.Build.cs b/Engine/Plugins/Experimental/NeuralRendering/Source/NeuralPostProcessing/NeuralPostProcessing.Build.cs
--- a/Engine/Plugins/Experimental/NeuralRendering/Source/NeuralPostProcessing/NeuralPostProcessing.Build.cs
+++ b/Engine/Plugins/Experimental/NeuralRendering/Source/NeuralPostProcessing/NeuralPostProcessing.Build.cs
@@ -14,14 +14,7 @@ public class NeuralPostProcessing : ModuleRules
 				// ... add public include paths required here ...
 			}
 			);
-				
 		
-		PrivateIncludePaths.AddRange(
-			new string[] {
-				// ... add other private include paths required here ...
-			}
-			);
-			
 		
 		PublicDependencyModuleNames.AddRange(
 			new string[]
diff --git a/Engine/Source/Runtime/Engine/Classes/Engine/NeuralProfile.h b/Engine/Source/Runtime/Engine/Classes/Engine/NeuralProfile.h
--- a/Engine/Source/Runtime/Engine/Classes/Engine/NeuralProfile.h
+++ b/Engine/Source/Runtime/Engine/Classes/Engine/NeuralProfile.h
@@ -78,7 +78,7 @@ struct FNeuralProfileStruct
 	ENeuralModelTileType TileSize;
 
 	/** Used to override the batch size if the batch dimension is dynamic (-1)*/
-	UPROPERTY(EditAnywhere, BlueprintReadOnly, Category = "Quality and Performance", meta = (EditCondition = "InputDimension[0] < 0"))
+	UPROPERTY(EditAnywhere, BlueprintReadOnly, Category = "Quality and Performance")
 	int32 BatchSizeOverride;
 
 	FNeuralProfileStruct()
