Add tile overlap support to texture indexing mode for neural post processing.
* Add `Auto` mode to `Tile size` to directly tile the screen image without scaling
* Add `Border Overlaps` controls between each tile to deal with discountinuity.
* Add `Overlap Resolve Type`:
  1. Ignore. Overlapped regions are ignored when concatenating tiles.
  2. Featherng. Overlapped regions are combined bilinearly

#tests test on style transfer and super resolution network on windows

diff --git a/Engine/Plugins/Experimental/NeuralRendering/Source/NeuralPostProcessing/Private/NeuralPostProcessModelInstance.cpp b/Engine/Plugins/Experimental/NeuralRendering/Source/NeuralPostProcessing/Private/NeuralPostProcessModelInstance.cpp
--- a/Engine/Plugins/Experimental/NeuralRendering/Source/NeuralPostProcessing/Private/NeuralPostProcessModelInstance.cpp
+++ b/Engine/Plugins/Experimental/NeuralRendering/Source/NeuralPostProcessing/Private/NeuralPostProcessModelInstance.cpp
@@ -80,9 +80,10 @@ UNeuralPostProcessModelInstance::UNeuralPostProcessModelInstance(FObjectInitiali
 	RDGOutputBuffer = nullptr;
 	RDGTiledInputBuffer = nullptr;
 	RDGTiledOutputBuffer = nullptr;
-	TileSize = 1;
+	DispatchSize = 1;
 	ModelTileSize = ENeuralModelTileType::OneByOne;
 	DimensionOverride = FIntVector4(-1);
+	TileOverlap = FIntPoint(0);
 }
 
 void UNeuralPostProcessModelInstance::Update(UNNEModelData* NNEModelData, FString RuntimeName)
@@ -90,11 +91,6 @@ void UNeuralPostProcessModelInstance::Update(UNNEModelData* NNEModelData, FStrin
 	this->CreateDefaultNNEModel(NNEModelData, RuntimeName);
 }
 
-void UNeuralPostProcessModelInstance::UpdateTileSize(int inTileSize)
-{
-	this->TileSize = inTileSize;
-}
-
 void UNeuralPostProcessModelInstance::Execute(FRDGBuilder& GraphBuilder)
 {
 	TArray<UE::NNE::FTensorBindingRDG> InputBindings;
@@ -105,37 +101,41 @@ void UNeuralPostProcessModelInstance::Execute(FRDGBuilder& GraphBuilder)
 	UE::NNE::FTensorShape InputShape = GetResolvedInputTensorShape();
 	int32 InputBatchSize = InputShape.GetData()[0];
 
-	auto RunSingleDispatch = [&]() {
+	auto RunSingleDispatch = [&](bool bAddFence = true) {
 		Input.Buffer = RDGInputBuffer;
 		Output.Buffer = RDGOutputBuffer;
 		int Status = ModelInstanceRDG->EnqueueRDG(GraphBuilder, InputBindings, OutputBindings);
 
 		checkf(!Status, TEXT("EnqueueRDG failed: %d"), Status);
 
-		AddPass(GraphBuilder, RDG_EVENT_NAME("EnqueueRDGFence"), [this](FRHICommandListImmediate& RHICmdList)
-			{
-				// Need to wait until the output buffer, RDGOutputBuffer, is ready.
-				RHICmdList.SubmitCommandsAndFlushGPU();
-				RHICmdList.BlockUntilGPUIdle();
-			});
+		if (bAddFence)
+		{
+			AddPass(GraphBuilder, RDG_EVENT_NAME("EnqueueRDGFence"), [this](FRHICommandListImmediate& RHICmdList)
+				{
+					// Need to wait until the output buffer, RDGOutputBuffer, is ready.
+					RHICmdList.SubmitCommandsAndFlushGPU();
+					RHICmdList.BlockUntilGPUIdle();
+				});
+		}
 	};
 
-	if (TileSize <= InputBatchSize)
+	if (DispatchSize == 1)
 	{
 		RunSingleDispatch();
 	}
-	else
+	else if (DispatchSize > 1)
 	{
 		// Since we have more tiles than batch per dispatch we need to copy the data from tiled
 		// to and from the buffer.
-		int32 NumOfDispatch = (TileSize + InputBatchSize - 1) / InputBatchSize;
+		int32 NumOfDispatch = DispatchSize;
 		int32 InputBufferSize = RDGInputBuffer->Desc.GetSize();
 		int32 OutputBufferSize = RDGOutputBuffer->Desc.GetSize();
 
 		for (int i = 0; i < NumOfDispatch; ++i)
 		{
+			bool bAddFence = i == NumOfDispatch - 1;
 			AddCopyBufferPass(GraphBuilder, RDGInputBuffer, 0, RDGTiledInputBuffer, InputBufferSize * i, InputBufferSize);
-			RunSingleDispatch();
+			RunSingleDispatch(bAddFence);
 			AddCopyBufferPass(GraphBuilder, RDGTiledOutputBuffer, OutputBufferSize * i, RDGOutputBuffer, 0, OutputBufferSize);
 		}
 	}
@@ -192,7 +192,7 @@ bool UNeuralPostProcessModelInstance::ModifyInputShape(int Dim, int Size)
 	if (NewResolvedInputShape[Dim] != Size)
 	{
 #if WITH_EDITOR
-		UE_LOG(LogNeuralPostProcessing, Error, TEXT("Cannot set dimension %d to %d. It is dynamic. revert back to %d"),Dim, Size, NewResolvedInputShape[Dim]);
+		UE_LOG(LogNeuralPostProcessing, Error, TEXT("Cannot set dimension %d to %d. It is not dynamic. revert back to %d"),Dim, Size, NewResolvedInputShape[Dim]);
 #endif
 		return false;
 	}
@@ -231,7 +231,7 @@ void UNeuralPostProcessModelInstance::CreateRDGBuffers(class FRDGBuilder& GraphB
 	RDGInputBuffer = GraphBuilder.CreateBuffer(InputBufferDesc, TEXT("NeuralPostProcessing.InputBuffer"));
 
 	// Create output buffers
-	checkf(ModelInstanceRDG->GetOutputTensorShapes().Num() == 1, TEXT("Post Processing requires models with a single output tensor!"));
+	checkf(ModelInstanceRDG->GetOutputTensorShapes().Num() >= 1, TEXT("Post Processing requires models with a single output tensor!"));
 	
 	ResolvedOutputTensorShape = ModelInstanceRDG->GetOutputTensorShapes()[0];
 	checkf(ResolvedOutputTensorShape.Rank() == 4, TEXT("Neural Post Processing requires models dim = 4 [static/dynamic x C x height x width]!"));
@@ -248,21 +248,22 @@ void UNeuralPostProcessModelInstance::CreateRDGBuffers(class FRDGBuilder& GraphB
 	
 	RDGOutputBuffer = GraphBuilder.CreateBuffer(OutputBufferDesc, *FString("SubsurfacePostProcessing.OutputBuffer"));
 
-	if (TileSize <= InputBatch)
+	if (DispatchSize <= 1)
 	{
 		RDGTiledInputBuffer = RDGInputBuffer;
 		RDGTiledOutputBuffer = RDGOutputBuffer;
 	}
-	else if (TileSize > InputBatch)
+	else
 	{
 		FRDGBufferDesc InputTiledBufferDesc = FRDGBufferDesc::CreateBufferDesc(sizeof(float), 
-			NeuralNetworkInputSize.X * NeuralNetworkInputSize.Y * InputChannels * InputBatch * TileSize);
+			NeuralNetworkInputSize.X * NeuralNetworkInputSize.Y * InputChannels * InputBatch * DispatchSize);
 		RDGTiledInputBuffer = GraphBuilder.CreateBuffer(InputTiledBufferDesc, TEXT("NeuralPostProcessing.TiledInputBuffer"));
 
 		FRDGBufferDesc OutputTiledBufferDesc = FRDGBufferDesc::CreateBufferDesc(sizeof(float), 
-			NeuralNetworkOutputSize.X * NeuralNetworkOutputSize.Y * OutputChannels * OutputBatch * TileSize);
+			NeuralNetworkOutputSize.X * NeuralNetworkOutputSize.Y * OutputChannels * OutputBatch * DispatchSize);
 		RDGTiledOutputBuffer = GraphBuilder.CreateBuffer(OutputTiledBufferDesc, TEXT("NeuralPostProcessing.TiledOutputBuffer"));
 	}
+
 }
 
 void UNeuralPostProcessModelInstance::CreateRDGBuffersIfNeeded(FRDGBuilder& GraphBuilder, bool bForceCreate)
diff --git a/Engine/Plugins/Experimental/NeuralRendering/Source/NeuralPostProcessing/Private/NeuralPostProcessModelInstance.h b/Engine/Plugins/Experimental/NeuralRendering/Source/NeuralPostProcessing/Private/NeuralPostProcessModelInstance.h
--- a/Engine/Plugins/Experimental/NeuralRendering/Source/NeuralPostProcessing/Private/NeuralPostProcessModelInstance.h
+++ b/Engine/Plugins/Experimental/NeuralRendering/Source/NeuralPostProcessing/Private/NeuralPostProcessModelInstance.h
@@ -19,6 +19,7 @@ TSharedPtr<UE::NNE::IModelInstanceRDG> CreateNNEModelInstance(UNNEModelData* NNE
 TSharedPtr<UE::NNE::IModelInstanceCPU> CreateNNECpuModelInstance(UNNEModelData* NNEModelData);
 
 enum class ENeuralModelTileType : uint8;
+enum class ETileOverlapResolveType : uint8;
 
 UCLASS()
 class UNeuralPostProcessModelInstance : public UObject
@@ -51,11 +52,16 @@ public:
 	FRDGBufferRef GetTiledInputBuffer();
 	FRDGBufferRef GetTiledOutputBuffer();
 
-	void UpdateTileSize(int TileSize = 1);
-	int GetTileSize() const { return TileSize; }
-
-	void UpdateModelTileType(ENeuralModelTileType inTileType) { ModelTileSize = inTileType; }
-	ENeuralModelTileType GetModelTileType()const { return ModelTileSize; }
+	void UpdateDispatchSize(int InDispatchSize){ DispatchSize = InDispatchSize;}
+	int	 GetDispatchSize() const{ return DispatchSize;}
+	void UpdateTileDimension(FIntPoint InTileDim) { TileDim = InTileDim; };
+	FIntPoint GetTileDimension() const { return TileDim; };
+	void UpdateModelTileType(ENeuralModelTileType InTileType) { ModelTileSize = InTileType; }
+	ENeuralModelTileType GetModelTileType() const { return ModelTileSize; }
+	void UpdateTileOverlap(FIntPoint InTileOverlap){ TileOverlap = InTileOverlap;}
+	FIntPoint GetTileOverlap() const { return TileOverlap;}
+	void UpdateTileOverlapResolveType(ETileOverlapResolveType InTileOverlapResolveType) { TileOverlapResolveType = InTileOverlapResolveType; }
+	ETileOverlapResolveType GetTileOverlapResolveType() const { return TileOverlapResolveType; }
 
 	bool IsValid() { return ModelInstanceRDG.IsValid();}
 
@@ -77,7 +83,8 @@ private:
 
 	ENeuralModelTileType ModelTileSize;
 
-	int TileSize = 1;
+	int DispatchSize = 1;
+	FIntPoint TileDim = FIntPoint(1,1);
 
 	// The NNE RDG Model 
 	TSharedPtr<UE::NNE::IModelInstanceRDG> ModelInstanceRDG;
@@ -89,5 +96,8 @@ private:
 	UE::NNE::FTensorShape ResolvedOutputTensorShape;
 
 	FIntVector4 DimensionOverride;
+
+	FIntPoint TileOverlap;
+	ETileOverlapResolveType TileOverlapResolveType;
 };
 
diff --git a/Engine/Plugins/Experimental/NeuralRendering/Source/NeuralPostProcessing/Private/NeuralPostProcessing.cpp b/Engine/Plugins/Experimental/NeuralRendering/Source/NeuralPostProcessing/Private/NeuralPostProcessing.cpp
--- a/Engine/Plugins/Experimental/NeuralRendering/Source/NeuralPostProcessing/Private/NeuralPostProcessing.cpp
+++ b/Engine/Plugins/Experimental/NeuralRendering/Source/NeuralPostProcessing/Private/NeuralPostProcessing.cpp
@@ -18,18 +18,63 @@
 DEFINE_LOG_CATEGORY(LogNeuralPostProcessing);
 #endif
 
+using namespace NeuralPostProcessng;
+
 namespace
 {
 	TAutoConsoleVariable<int32> CVarNeuralPostProcessApply(
-		TEXT("r.neuralpostprocess.apply"),
+		TEXT("r.Neuralpostprocess.Apply"),
 		1,
 		TEXT(" 0: disabled\n")
 		TEXT(" 1: enabled (default)"),
 		ECVF_RenderThreadSafe | ECVF_Scalability);
 
+	TAutoConsoleVariable<int32> CVarNeuralPostProcessTileOverlap(
+		TEXT("r.Neuralpostprocess.TileOverlap"),
+		-1,
+		TEXT(" <0: Use the overlap from the profile\n"),
+		ECVF_RenderThreadSafe | ECVF_Scalability);
+
+	TAutoConsoleVariable<int32> CVarNeuralPostProcessTileOverlapResolveType(
+		TEXT("r.Neuralpostprocess.TileOverlap.ResolveType"),
+		-1,
+		TEXT(" 0: Ignore overlapped region when concatinating the resolved subframe.\n")
+		TEXT(" 1: Feathering overlapping regions based on bilinear filtering. Works best when Input and Output dimension are the same\n")
+		TEXT(" <0: Use the ResolveType from profile."),
+		ECVF_RenderThreadSafe | ECVF_Scalability);
+
+	TAutoConsoleVariable<int32> CVarNeuralPostProcessTileOverlapVisualize(
+		TEXT("r.Neuralpostprocess.TileOverlap.Visualize"),
+		0,
+		TEXT(" Visualize the overlap region if not zero\n"),
+		ECVF_RenderThreadSafe | ECVF_Scalability);
+
+	TAutoConsoleVariable<float> CVarNeuralPostProcessTileOverlapVisualizeIntensity(
+		TEXT("r.Neuralpostprocess.TileOverlap.Visualize.Intensity"),
+		1.0f,
+		TEXT(" Adjust the intensity of the overlap for visualization\n"),
+		ECVF_RenderThreadSafe | ECVF_Scalability);
+
 	DECLARE_GPU_STAT(NeuralPostProcessing)
 }
 
+static FIntPoint ClampOverlap(FIntPoint TileOverlap,FIntPoint BufferSizeXY)
+{
+	// Use TileOverlap from CVar if larger or equal than zero.
+	int32 OverlapFromCVar = CVarNeuralPostProcessTileOverlap.GetValueOnRenderThread();
+	bool ShouldOverrideOverlapWithCVar = OverlapFromCVar >= 0;
+	if (ShouldOverrideOverlapWithCVar)
+	{
+		TileOverlap = FIntPoint(OverlapFromCVar);
+	}
+
+	// Max tile border is 1/4 of the dimension. 
+	FIntPoint MaxTileOverlap = BufferSizeXY / 4;
+	TileOverlap.X = FMath::Clamp(TileOverlap.X, 0, MaxTileOverlap.X);
+	TileOverlap.Y = FMath::Clamp(TileOverlap.Y, 0, MaxTileOverlap.Y);
+
+	return TileOverlap;
+}
 
 /**
 *
@@ -107,6 +152,7 @@ public:
 #if WITH_EDITOR
 			UE_LOG(LogNeuralPostProcessing, Error, TEXT("NNEModelData is invalid at Slot %d."), AllocationId);
 #endif
+			FNeuralPostProcessModelInstanceManager::Get()->SetModelInstance(AllocationId, nullptr);
 			return;
 		}
 
@@ -119,13 +165,41 @@ public:
 	virtual void UpdateTileType(int32 AllocationId, ENeuralModelTileType ModelTileType) override
 	{
 		UNeuralPostProcessModelInstance* Instance = FNeuralPostProcessModelInstanceManager::Get()->GetModelInstance(AllocationId);
-		Instance->UpdateModelTileType(ModelTileType);
+		if (Instance)
+		{
+			Instance->UpdateModelTileType(ModelTileType);
+		}
 	}
 
 	virtual bool UpdateBatchSize(int32 AllocationId, int32 BatchSize) override
 	{
 		UNeuralPostProcessModelInstance* Instance = FNeuralPostProcessModelInstanceManager::Get()->GetModelInstance(AllocationId);
-		return Instance->ModifyInputShape(0, BatchSize);
+		if (Instance)
+		{
+			return Instance->ModifyInputShape(0, BatchSize);
+		}
+
+		return false;
+	}
+
+	virtual void UpdateTileOverlap(int32 AllocationId, FIntPoint TileOverlap) override
+	{
+		UNeuralPostProcessModelInstance* Instance = FNeuralPostProcessModelInstanceManager::Get()->GetModelInstance(AllocationId);
+		if (Instance)
+		{
+			UE::NNE::FTensorShape InputShape = Instance->GetResolvedInputTensorShape();
+			FIntPoint BufferSizeXY = FIntPoint((int32)InputShape.GetData()[3], (int32)InputShape.GetData()[2]);
+			Instance->UpdateTileOverlap(ClampOverlap(TileOverlap, BufferSizeXY));
+		}
+	}
+
+	virtual void UpdateTileOverlapResolveType(int32 AllocationId, ETileOverlapResolveType TileOverlapResolveType) override
+	{
+		UNeuralPostProcessModelInstance* Instance = FNeuralPostProcessModelInstanceManager::Get()->GetModelInstance(AllocationId);
+		if (Instance)
+		{
+			Instance->UpdateTileOverlapResolveType(TileOverlapResolveType);
+		}
 	}
 
 	virtual void RemoveModel(int32 AllocationId) override
@@ -212,9 +286,9 @@ int32 GetTotalModelTileCount(ENeuralModelTileType ModelTileSize)
 	case ENeuralModelTileType::EightByEight:
 		TotalTileCount = 8 * 8;
 		break;
-	//case ENeuralModelTileType::Auto:
-	//	TotalTileCount = -1;	// determined at runtime by the viewport size and W,H dimension of the network 
-	//	break;
+	case ENeuralModelTileType::Auto:
+		TotalTileCount = -1;	// determined at runtime by the viewport size and W,H dimension of the network 
+		break;
 	case ENeuralModelTileType::OneByOne:
 	default:
 		break;
@@ -243,33 +317,101 @@ static void AllocateInputBuffer_RenderingThread(
 
 	UE::NNE::FTensorShape InputShape = Model->GetResolvedInputTensorShape();
 	InputBufferDimension = FVector4f(InputShape.GetData()[0], InputShape.GetData()[1], InputShape.GetData()[2], InputShape.GetData()[3]);
-	
+	ENeuralModelTileType ModelTileType = Model->GetModelTileType();
+	FScreenPassTextureViewportParameters ViewportParameters = GetScreenPassTextureViewportParameters(Viewport);
+	int32 BatchDim = InputBufferDimension.X;
+
 	// Calculate the number of tiles
-	int NumTiles = 1;
-	{
-		ENeuralModelTileType ModelTileSize = Model->GetModelTileType();
-		// @todo: Finish the implement of auto
-		//if (ModelTileSize == ENeuralModelTileType::Auto)
-		//{
-		//	// allocate the number of tiles based on the size of the viewport and the size of the buffer.
-		//	FScreenPassTextureViewportParameters ViewportParameters = GetScreenPassTextureViewportParameters(Viewport);
-		//	FVector2f TileSizeWH = FMath::DivideAndRoundUp(ViewportParameters.ViewportSize, FVector2f(InputBufferDimension.W, InputBufferDimension.Z));
-		//	float BatchDim = InputBufferDimension.X;
-		//	NumTiles = FMath::DivideAndRoundUp(TileSizeWH.X * TileSizeWH.Y, BatchDim);
-		//}
-		//else
+	int NumDispatches = 1;
+	FIntPoint TileDimension = FIntPoint(1, 1);
+	{	
+		if (ModelTileType == ENeuralModelTileType::Auto)
 		{
-			NumTiles = GetTotalModelTileCount(ModelTileSize);
+			// Allocate the number of tiles based on the size of the viewport, the size of the buffer and tile configurations.
+			FIntPoint InputBufferSizeXY = FIntPoint(InputBufferDimension.W, InputBufferDimension.Z);
+			FIntPoint TileOvelapBorder = ClampOverlap(Model->GetTileOverlap(), InputBufferSizeXY);
+			FIntPoint EffectiveBufferWH = InputBufferSizeXY - TileOvelapBorder * 2;
+			FVector2f TileSizeWH = FMath::DivideAndRoundUp(ViewportParameters.ViewportSize.IntPoint(), EffectiveBufferWH);
+			TileDimension = FIntPoint(TileSizeWH.X,TileSizeWH.Y);
+			NumDispatches = FMath::DivideAndRoundUp(TileDimension.X * TileDimension.Y, BatchDim);
+		}
+		else
+		{
+			// Use fixed number of tiles
+			int32 TileCount = GetTotalModelTileCount(ModelTileType);
+			int32 SideLength = (int32)FMath::Sqrt((float)TileCount);
+			TileDimension = FIntPoint(SideLength, SideLength);
+			NumDispatches = FMath::DivideAndRoundUp(TileCount, BatchDim);
 		}
 
-		Model->UpdateTileSize(NumTiles);
+		Model->UpdateDispatchSize(NumDispatches);
+		Model->UpdateTileDimension(TileDimension);
 	}
-
+	
+	// Create the buffer
 	Model->CreateRDGBuffersIfNeeded(GraphBuilder, true);
 
+
 	// Output the buffer and dimension for use.
 	InputNeuralBuffer = Model->GetTiledInputBuffer();
-	InputBufferDimension.X *= NumTiles;
+	InputBufferDimension.X *= NumDispatches;
+
+	// Clear the input buffer at the border if we use ENeuralModelTileType::Auto
+	// and the botom and right tiles partially covers the outside.
+	if (ModelTileType == ENeuralModelTileType::Auto)
+	{
+		//@TODO: only clear part of the regions at borders.
+		AddClearUAVPass(GraphBuilder,GraphBuilder.CreateUAV(InputNeuralBuffer,PF_R32_FLOAT),0);
+	}
+}
+
+FCopyBetweenTextureAndOverlappedTileBufferCS::EOverlapResolveType GetOverlapResolveType(FIntPoint TileOverlap, ETileOverlapResolveType TileOverlapResolveType)
+{
+	typedef FCopyBetweenTextureAndOverlappedTileBufferCS SHADER;
+
+	SHADER::EOverlapResolveType OverlapResolveType = SHADER::EOverlapResolveType::Ignore;
+
+	// Set TileOverlap Resolve type by profile config
+	switch (TileOverlapResolveType)
+	{
+	case ETileOverlapResolveType::Feathering:
+		OverlapResolveType = SHADER::EOverlapResolveType::Feathering;
+		break;
+	case ETileOverlapResolveType::Ignore:
+	default:
+		break;
+	}
+
+	// If both overlaps are zero, switch to Ignore for performance
+	if (OverlapResolveType == SHADER::EOverlapResolveType::Feathering && (TileOverlap.X == 0 && TileOverlap.Y == 0))
+	{
+		OverlapResolveType = SHADER::EOverlapResolveType::Ignore;
+	}
+
+	// Override by console variables
+	int32 OverlapResolveTypeOverride = CVarNeuralPostProcessTileOverlapResolveType.GetValueOnRenderThread();
+	if (OverlapResolveTypeOverride == 0)
+	{
+		OverlapResolveType = SHADER::EOverlapResolveType::Ignore;
+	}
+	else if (OverlapResolveTypeOverride > 0)
+	{
+		OverlapResolveType = SHADER::EOverlapResolveType::Feathering;
+	}
+
+	return OverlapResolveType;
+}
+
+FIntPoint ScaleValueFromSourceToTargetSize(FIntPoint Value, FIntPoint SourceSize, FIntPoint TargetSize)
+{
+	FVector2f ScaledValue = FVector2f(Value.X * TargetSize.X * 1.0f / SourceSize.X, Value.Y * TargetSize.Y * 1.0f / SourceSize.Y);
+	return FIntPoint(FMath::CeilToInt(ScaledValue.X), FMath::CeilToInt(ScaledValue.Y));
+}
+
+FIntPoint GetScaledTileOverlap(FIntPoint TileOverlap, FIntPoint SourceSize, FIntPoint TargetSize)
+{
+	// Scaled overlap is designed to return integer.
+	return ScaleValueFromSourceToTargetSize(TileOverlap, SourceSize, TargetSize);
 }
 
 static void	ApplyNeuralNetworks_RenderingThread(
@@ -296,9 +438,46 @@ static void	ApplyNeuralNetworks_RenderingThread(
 	FIntPoint TextureSize = NeuralTexture->Desc.Extent;
 	FGlobalShaderMap* GlobalShaderMap = GetGlobalShaderMap(GMaxRHIFeatureLevel);
 	float Scale = 1.0;
+
+	UE::NNE::FTensorShape InputShape = Model->GetResolvedInputTensorShape();
+
+	FIntPoint NeuralNetworkInputSize = { (int32)InputShape.GetData()[3], (int32)InputShape.GetData()[2] };// Width, height
+	FIntPoint TileOverlap = ClampOverlap(Model->GetTileOverlap(), NeuralNetworkInputSize);
+	FIntPoint TileDimension = Model->GetTileDimension();
 	
 	const FScreenPassTextureViewport NeuralPostProcessViewport(NeuralTexture, ViewRect);
 	const FScreenPassTextureViewportParameters InputViewportParameters = GetScreenPassTextureViewportParameters(NeuralPostProcessViewport);
+	
+	const ENeuralModelTileType ModelTileType = Model->GetModelTileType();
+	const bool bRequireInputBufferScale = ModelTileType != ENeuralModelTileType::Auto;
+	const ETileOverlapResolveType TileOverlapResolveType = Model->GetTileOverlapResolveType();
+
+	struct FIntermediateTexture
+	{
+		FRDGTextureRef Texture;
+		FIntPoint		Extent;
+	};
+
+	auto GetIntermediateTexture = [&](FIntPoint InNetworkSize, FIntPoint InOverLap, bool bRequireBufferScale, const TCHAR* Name)->FIntermediateTexture {
+		FIntermediateTexture InputIntermediateTexture;
+
+		if (bRequireBufferScale)
+		{
+			FIntPoint IntermediateTextureExtent = (InNetworkSize - InOverLap * 2) * TileDimension;
+			FRDGTextureDesc IntermediateTextureDesc = NeuralTexture->Desc;
+			IntermediateTextureDesc.Extent = IntermediateTextureExtent;
+
+			// @TODO: conditionally create the intermediate buffer texture
+			InputIntermediateTexture.Texture = GraphBuilder.CreateTexture(IntermediateTextureDesc, Name);
+			InputIntermediateTexture.Extent = IntermediateTextureExtent;
+		}
+		else
+		{
+			InputIntermediateTexture.Texture = NeuralTexture;
+			InputIntermediateTexture.Extent = InputViewportParameters.ViewportSize.IntPoint();
+		}
+		return InputIntermediateTexture;
+	};
 
 	// 1. Preprocess the input data by copying from Texture to buffer if required
 	{
@@ -306,49 +485,92 @@ static void	ApplyNeuralNetworks_RenderingThread(
 		FRDGBufferRef IndirectDispatchBuffer = GraphBuilder.CreateBuffer(
 			FRDGBufferDesc::CreateIndirectDesc<FRHIDispatchIndirectParameters>(1), TEXT("NeuralPostProcessing.IndirectDispatchBuffer"));
 
-		UE::NNE::FTensorShape InputShape = Model->GetResolvedInputTensorShape();
-		FIntPoint NeuralNetworkInputSize = { (int32)InputShape.GetData()[3], (int32)InputShape.GetData()[2] };// Width, height
-		int32 TileSize = Model->GetTileSize();
+		FIntermediateTexture IntermediateTexture = 
+			GetIntermediateTexture(NeuralNetworkInputSize, TileOverlap, bRequireInputBufferScale, TEXT("NeuralPostProcessing.IntermediateBufferTexture"));
 
-		// Build the indirect dispatch parameters with InputSourceType
+		// Downscale
+		if (bRequireInputBufferScale)
 		{
-			typedef FNeuralPostProcessingBuildIndirectDispatchArgsCS ARGSETUPSHADER;
-			ARGSETUPSHADER::FParameters* PassParameters = GraphBuilder.AllocParameters<ARGSETUPSHADER::FParameters>();
-			PassParameters->NetworkTextureSize = NeuralNetworkInputSize;
-			PassParameters->SourceType = GraphBuilder.CreateSRV(InputSourceType, EPixelFormat::PF_R32_UINT);
-			PassParameters->RWIndirectDispatchArgsBuffer = GraphBuilder.CreateUAV(IndirectDispatchBuffer, EPixelFormat::PF_R32_UINT);
-			TShaderMapRef<ARGSETUPSHADER> ComputeShader(GlobalShaderMap);
-			FComputeShaderUtils::AddPass(GraphBuilder, FRDGEventName(TEXT("NeuralPostProcessing::BuildIndirectArgs(Dispatch)")), ComputeShader, PassParameters, FIntVector(1,1,1));
-		}
+			// Build the indirect dispatch parameters with InputSourceType
+			{
+				typedef FNeuralPostProcessingBuildIndirectDispatchArgsCS ARGSETUPSHADER;
+				ARGSETUPSHADER::FParameters* PassParameters = GraphBuilder.AllocParameters<ARGSETUPSHADER::FParameters>();
+				PassParameters->TargetDimension = IntermediateTexture.Extent;
+				PassParameters->SourceType = GraphBuilder.CreateSRV(InputSourceType, EPixelFormat::PF_R32_UINT);
+				PassParameters->RWIndirectDispatchArgsBuffer = GraphBuilder.CreateUAV(IndirectDispatchBuffer, EPixelFormat::PF_R32_UINT);
+				TShaderMapRef<ARGSETUPSHADER> ComputeShader(GlobalShaderMap);
+				FComputeShaderUtils::AddPass(GraphBuilder, FRDGEventName(TEXT("NeuralPostProcessing::BuildIndirectArgs(Dispatch)")), ComputeShader, PassParameters, FIntVector(1, 1, 1));
+			}
 
-		// Use the neural network's input dimension configuration if available, otherwise use the dynamic size of the buffer
-		if (TileSize == 1)
-		{
-			typedef FNeuralPostProcessingPrepareInputCS SHADER;
+			typedef FDownScaleTextureCS SHADER;
 			SHADER::FParameters* PassParameters = GraphBuilder.AllocParameters<SHADER::FParameters>();
 			{
-				PassParameters->Input0 = GetNeuralPostProcessInput(NeuralTexture, InputViewportParameters);
-				PassParameters->InputTextureSampler = TStaticSamplerState<SF_Bilinear, AM_Clamp, AM_Clamp, AM_Clamp>::GetRHI();
-				PassParameters->InputBufferWidth = NeuralNetworkInputSize.X;
-				PassParameters->InputBufferHeight = NeuralNetworkInputSize.Y;
-				PassParameters->InputBuffer = GraphBuilder.CreateUAV(FRDGBufferUAVDesc(Model->GetInputBuffer(), PF_R32_FLOAT));
-				PassParameters->ColorScale = Scale;
+				PassParameters->Source = GetNeuralPostProcessInput(NeuralTexture, InputViewportParameters);
+				PassParameters->SourceTextureSampler = TStaticSamplerState<SF_Bilinear, AM_Clamp, AM_Clamp, AM_Clamp>::GetRHI();
+				PassParameters->TargetWidth = IntermediateTexture.Extent.X;
+				PassParameters->TargetHeight = IntermediateTexture.Extent.Y;
+				PassParameters->TargetTexture = GraphBuilder.CreateUAV(FRDGTextureUAVDesc(IntermediateTexture.Texture));
 				PassParameters->IndirectDispatchArgsBuffer = IndirectDispatchBuffer;
 			}
-
-			TShaderMapRef<FNeuralPostProcessingPrepareInputCS> ComputeShader(GlobalShaderMap);
+			TShaderMapRef<SHADER> ComputeShader(GlobalShaderMap);
 			FComputeShaderUtils::AddPass(
 				GraphBuilder,
-				RDG_EVENT_NAME("NeuralPostProcessing::PrepareInput"),
+				RDG_EVENT_NAME("NeuralPostProcessing::DownScale (%dx%d -> %dx%d)",
+					(int32)InputViewportParameters.ViewportSize.X,
+					(int32)InputViewportParameters.ViewportSize.Y,
+					IntermediateTexture.Extent.X,
+					IntermediateTexture.Extent.Y),
 				ERDGPassFlags::Compute, //| ERDGPassFlags::NeverCull
 				ComputeShader,
 				PassParameters,
-				PassParameters->IndirectDispatchArgsBuffer,0);
+				PassParameters->IndirectDispatchArgsBuffer, 0);
 		}
-		else
+
+		// Copy from Texture to Overlapped tile buffer
 		{
-			// @TODO: Implement texture copy when neural network buffer is accessed by texture index.
+			int32 NumOfChannel = (int32)InputShape.GetData()[1];
+			// Build the indirect dispatch parameters with InputSourceType
+			{
+				typedef FNeuralPostProcessingBuildIndirectDispatchArgsCS ARGSETUPSHADER;
+				ARGSETUPSHADER::FParameters* PassParameters = GraphBuilder.AllocParameters<ARGSETUPSHADER::FParameters>();
+				PassParameters->TargetDimension = NeuralNetworkInputSize * TileDimension;
+				PassParameters->SourceType = GraphBuilder.CreateSRV(InputSourceType, EPixelFormat::PF_R32_UINT);
+				PassParameters->RWIndirectDispatchArgsBuffer = GraphBuilder.CreateUAV(IndirectDispatchBuffer, EPixelFormat::PF_R32_UINT);
+				TShaderMapRef<ARGSETUPSHADER> ComputeShader(GlobalShaderMap);
+				FComputeShaderUtils::AddPass(GraphBuilder, FRDGEventName(TEXT("NeuralPostProcessing::BuildIndirectArgs(Dispatch)")), ComputeShader, PassParameters, FIntVector(1, 1, 1));
+			}
+
+			typedef FCopyBetweenTextureAndOverlappedTileBufferCS SHADER;
+			SHADER::FParameters* PassParameters = GraphBuilder.AllocParameters<SHADER::FParameters>();
+			{
+				PassParameters->SourceWidth = IntermediateTexture.Extent.X;
+				PassParameters->SourceHeight = IntermediateTexture.Extent.Y;
+				PassParameters->RWSourceTexture = GraphBuilder.CreateUAV(FRDGTextureUAVDesc(IntermediateTexture.Texture));
+				PassParameters->SourceTextureSampler = TStaticSamplerState<SF_Point, AM_Mirror, AM_Mirror, AM_Mirror>::GetRHI();
+				
+				PassParameters->TargetOverlappedTileWidth = NeuralNetworkInputSize.X;
+				PassParameters->TargetOverlappedTileHeight = NeuralNetworkInputSize.Y;
+				PassParameters->ViewTileDimension = TileDimension;
+				PassParameters->TileOverlap = FVector2f(TileOverlap.X,TileOverlap.Y);
+				PassParameters->NumOfChannel = NumOfChannel;
+				PassParameters->TargetBuffer = GraphBuilder.CreateUAV(FRDGBufferUAVDesc(Model->GetTiledInputBuffer(), PF_R32_FLOAT));
+				PassParameters->IndirectDispatchArgsBuffer = IndirectDispatchBuffer;
+			}
+
+			SHADER::FPermutationDomain ComputeShaderPermutationVector;
+			ComputeShaderPermutationVector.Set<SHADER::FDimensionCopyDirection>(SHADER::EDirection::ToOverlappedTiles);
+			ComputeShaderPermutationVector.Set<SHADER::FDimensionOverlapResolveType>(SHADER::EOverlapResolveType::Ignore);
+			TShaderMapRef<SHADER> ComputeShader(GlobalShaderMap,ComputeShaderPermutationVector);
+
+			FComputeShaderUtils::AddPass(
+				GraphBuilder,
+				RDG_EVENT_NAME("NeuralPostProcessing::CopyToTiledBuffer (TileOverlap=%d,%d)",TileOverlap.X,TileOverlap.Y),
+				ERDGPassFlags::Compute,
+				ComputeShader,
+				PassParameters,
+				PassParameters->IndirectDispatchArgsBuffer, 0);
 		}
+
 	}
 
 	// 2. Run the network
@@ -367,44 +589,138 @@ static void	ApplyNeuralNetworks_RenderingThread(
 
 		return OutputShape.GetData()[0] == InputShape.GetData()[0] &&
 			OutputShape.GetData()[1] == InputShape.GetData()[1] &&
-			Model->GetTileSize() == 1; //@TODO: add texture copy when tile size is larger than batch size. 
+			Model->GetDispatchSize() >= 1;
 	};
 
 	if (ShouldUpdateNeuralTexture())
 	{
-		FRDGBufferRef OutputBuffer = bNeuralPostProcessApply ? Model->GetOutputBuffer() : Model->GetInputBuffer();
+		FRDGBufferRef OutputBuffer = bNeuralPostProcessApply ? Model->GetTiledOutputBuffer() : Model->GetTiledInputBuffer();
 		FRDGBufferUAVRef OutputBufferUAV = GraphBuilder.CreateUAV(FRDGBufferUAVDesc(OutputBuffer, PF_R32_FLOAT));
 
 		UE::NNE::FTensorShape OutputShape = bNeuralPostProcessApply ? Model->GetResolvedOutputTensorShape() : Model->GetResolvedInputTensorShape();
 		FIntPoint NeuralNetworkOutputSize = { (int32)OutputShape.GetData()[3], (int32)OutputShape.GetData()[2] };
 
-		const FIntRect Viewport = FIntRect(0,TextureSize);
+		FIntPoint ScaledTileOverlap = GetScaledTileOverlap(TileOverlap, NeuralNetworkInputSize, NeuralNetworkOutputSize);
 
-		FNeuralPostProcessingProcessOutputPS::FParameters* ProcessOutputParameters = GraphBuilder.AllocParameters<FNeuralPostProcessingProcessOutputPS::FParameters>();
-		ProcessOutputParameters->Input0.Texture = nullptr;
-		ProcessOutputParameters->Input0.Viewport = InputViewportParameters;
-		ProcessOutputParameters->OutputBufferWidth = NeuralNetworkOutputSize.X;
-		ProcessOutputParameters->OutputBufferHeight = NeuralNetworkOutputSize.Y;
-		ProcessOutputParameters->OutputBuffer = OutputBufferUAV;
-		ProcessOutputParameters->ColorScale = Scale;
-		ProcessOutputParameters->RenderTargets[0] = FRenderTargetBinding(NeuralTexture, ERenderTargetLoadAction::ELoad);
+		// Need to scale output buffer if the output network size is larger than that of the input of the network.
+		const bool bRequireOutputBufferScale = ModelTileType != ENeuralModelTileType::Auto ||
+			(NeuralNetworkOutputSize.X != NeuralNetworkInputSize.X || NeuralNetworkOutputSize.Y != NeuralNetworkInputSize.Y);
 
-		TShaderMapRef<FNeuralPostProcessingProcessOutputPS> WriteOutputShader(GlobalShaderMap);
+		FIntermediateTexture IntermediateTexture = 
+			GetIntermediateTexture(NeuralNetworkOutputSize, ScaledTileOverlap, bRequireOutputBufferScale, TEXT("NeuralPostProcessing.IntermediateOutputBufferTexture"));
+		
+		// Copy from overlapped tile buffer to texture
+		{
+			int32 NumOfChannel = (int32)InputShape.GetData()[1];
 
-		FPixelShaderUtils::AddFullscreenPass(
-			GraphBuilder,
-			GlobalShaderMap,
-			RDG_EVENT_NAME("NeuralPostProcessing::ProcessOutput"),
-			WriteOutputShader,
-			ProcessOutputParameters,
-			ViewRect);
+			typedef FCopyBetweenTextureAndOverlappedTileBufferCS SHADER;
+			SHADER::FParameters* PassParameters = GraphBuilder.AllocParameters<SHADER::FParameters>();
+			{
+				PassParameters->SourceWidth = IntermediateTexture.Extent.X;
+				PassParameters->SourceHeight = IntermediateTexture.Extent.Y;
+				PassParameters->RWSourceTexture = GraphBuilder.CreateUAV(FRDGTextureUAVDesc(IntermediateTexture.Texture));
+				PassParameters->SourceTextureSampler = TStaticSamplerState<SF_Point, AM_Mirror, AM_Mirror, AM_Mirror>::GetRHI();
+
+				PassParameters->TargetOverlappedTileWidth = NeuralNetworkOutputSize.X;
+				PassParameters->TargetOverlappedTileHeight = NeuralNetworkOutputSize.Y;
+				PassParameters->ViewTileDimension = TileDimension;
+				PassParameters->TileOverlap = FVector2f(ScaledTileOverlap.X, ScaledTileOverlap.Y);
+				PassParameters->NumOfChannel = NumOfChannel;
+				PassParameters->bVisualizeOverlap = CVarNeuralPostProcessTileOverlapVisualize.GetValueOnRenderThread() != 0;
+				PassParameters->OverlapVisualizeIntensity = FMath::Max(0.0001, CVarNeuralPostProcessTileOverlapVisualizeIntensity.GetValueOnRenderThread());
+				PassParameters->TargetBuffer = OutputBufferUAV;
+			}
+
+			SHADER::FPermutationDomain ComputeShaderPermutationVector;
+			ComputeShaderPermutationVector.Set<SHADER::FDimensionCopyDirection>(SHADER::EDirection::FromOverlappedTiles);
+			ComputeShaderPermutationVector.Set<SHADER::FDimensionOverlapResolveType>(GetOverlapResolveType(TileOverlap, TileOverlapResolveType));
+			TShaderMapRef<SHADER> ComputeShader(GlobalShaderMap, ComputeShaderPermutationVector);
+
+			FComputeShaderUtils::AddPass(
+				GraphBuilder,
+				RDG_EVENT_NAME("NeuralPostProcessing::CopyFromTiledBuffer (TileOverlap=%d,%d)", TileOverlap.X,TileOverlap.Y),
+				ERDGPassFlags::Compute,
+				ComputeShader,
+				PassParameters,
+				FComputeShaderUtils::GetGroupCount(NeuralNetworkOutputSize * TileDimension, NEURAL_POST_PROCESSING_THREAD_GROUP_SIZE));
+		}
+
+		// Scale and write to the target texture
+		if (bRequireOutputBufferScale)
+		{
+			FIntPoint SourceExtent = IntermediateTexture.Extent;
+			
+			// Update the effective source width and height if we use Auto.
+			if (ModelTileType == ENeuralModelTileType::Auto)
+			{
+				int ModX = InputViewportParameters.ViewportSize.IntPoint().X % (NeuralNetworkInputSize.X - 2 * TileOverlap.X);
+				int ModY = InputViewportParameters.ViewportSize.IntPoint().Y % (NeuralNetworkInputSize.Y - 2 * TileOverlap.Y);
+				FIntPoint ScaledSize = ScaleValueFromSourceToTargetSize(FIntPoint(ModX, ModY), NeuralNetworkInputSize, NeuralNetworkOutputSize);
+				SourceExtent = (NeuralNetworkOutputSize - ScaledTileOverlap * 2 )* (TileDimension - 1) + ScaledSize;
+			}
+
+			const bool bUpScale = SourceExtent[0] < ViewRect.Size()[0] && SourceExtent[1] < ViewRect.Size()[1];
+
+			if (bUpScale)
+			{
+				typedef FUpscaleTexture SHADER;
+				SHADER::FParameters* ProcessOutputParameters = GraphBuilder.AllocParameters<SHADER::FParameters>();
+				ProcessOutputParameters->Source_Texture = IntermediateTexture.Texture;
+				ProcessOutputParameters->ViewportSize = InputViewportParameters.ViewportSize.IntPoint();
+				ProcessOutputParameters->SourceWidth = SourceExtent.X;
+				ProcessOutputParameters->SourceHeight = SourceExtent.Y;
+				ProcessOutputParameters->RenderTargets[0] = FRenderTargetBinding(NeuralTexture, ERenderTargetLoadAction::ELoad);
+
+				TShaderMapRef<SHADER> WriteOutputShader(GlobalShaderMap);
+
+				FPixelShaderUtils::AddFullscreenPass(
+					GraphBuilder,
+					GlobalShaderMap,
+					RDG_EVENT_NAME("NeuralPostProcessing::Upscale (%dx%d -> %dx%d)",
+						SourceExtent.X,
+						SourceExtent.Y,
+						(int32)InputViewportParameters.ViewportSize.X,
+						(int32)InputViewportParameters.ViewportSize.Y),
+					WriteOutputShader,
+					ProcessOutputParameters,
+					ViewRect);
+			}
+			else // Down Scale
+			{
+				FScreenPassTextureViewport IntermediateTextureViewport(IntermediateTexture.Texture, FIntRect(0, 0, SourceExtent.X, SourceExtent.Y));
+				FScreenPassTextureViewportParameters IntermediateTextureViewportParameters = GetScreenPassTextureViewportParameters(IntermediateTextureViewport);
+
+				typedef FDownScaleTexture SHADER;
+				SHADER::FParameters* PassParameters = GraphBuilder.AllocParameters<SHADER::FParameters>();
+				{
+					PassParameters->Source = GetNeuralPostProcessInput(IntermediateTexture.Texture, IntermediateTextureViewportParameters);
+					PassParameters->SourceTextureSampler = TStaticSamplerState<SF_Point, AM_Clamp, AM_Clamp, AM_Clamp>::GetRHI();
+					PassParameters->TargetWidth = InputViewportParameters.ViewportSize.X;
+					PassParameters->TargetHeight = InputViewportParameters.ViewportSize.Y;
+					PassParameters->RenderTargets[0] = FRenderTargetBinding(NeuralTexture, ERenderTargetLoadAction::ELoad);
+				}
+				TShaderMapRef<SHADER> WriteOutputShader(GlobalShaderMap);
+
+				FPixelShaderUtils::AddFullscreenPass(
+					GraphBuilder,
+					GlobalShaderMap,
+					RDG_EVENT_NAME("NeuralPostProcessing::Downscale (%dx%d -> %dx%d)",
+						SourceExtent.X,
+						SourceExtent.Y,
+						(int32)InputViewportParameters.ViewportSize.X,
+						(int32)InputViewportParameters.ViewportSize.Y),
+					WriteOutputShader,
+					PassParameters,
+					ViewRect);
+			}
+		}
 	}
 
 	// 4. Read back the buffer, so the user can directly decode in the post process material.
 	OutputNeuralBuffer = bNeuralPostProcessApply ? Model->GetTiledOutputBuffer() : Model->GetTiledInputBuffer();
 	UE::NNE::FTensorShape OutputShape = bNeuralPostProcessApply ? Model->GetResolvedOutputTensorShape() : Model->GetResolvedInputTensorShape();		
 	BufferDimension = FVector4f(OutputShape.GetData()[0], OutputShape.GetData()[1], OutputShape.GetData()[2], OutputShape.GetData()[3]);
-	BufferDimension.X *= Model->GetTileSize();
+	BufferDimension.X *= Model->GetDispatchSize();
 }
 
 class FNeuralPostProcess : public INeuralPostProcessInterface
diff --git a/Engine/Plugins/Experimental/NeuralRendering/Source/NeuralPostProcessing/Private/NeuralPostProcessingCS.h b/Engine/Plugins/Experimental/NeuralRendering/Source/NeuralPostProcessing/Private/NeuralPostProcessingCS.h
--- a/Engine/Plugins/Experimental/NeuralRendering/Source/NeuralPostProcessing/Private/NeuralPostProcessingCS.h
+++ b/Engine/Plugins/Experimental/NeuralRendering/Source/NeuralPostProcessing/Private/NeuralPostProcessingCS.h
@@ -9,65 +9,129 @@
 
 #define NEURAL_POST_PROCESSING_THREAD_GROUP_SIZE 32
 
-BEGIN_SHADER_PARAMETER_STRUCT(FNueralPostProcessInput, )
-	SHADER_PARAMETER_STRUCT_INCLUDE(FScreenPassTextureViewportParameters, Viewport)
-	SHADER_PARAMETER_RDG_TEXTURE(Texture2D, Texture)
-END_SHADER_PARAMETER_STRUCT()
-
-FNueralPostProcessInput GetNeuralPostProcessInput(FRDGTextureRef Texture, const FScreenPassTextureViewportParameters& ViewportParameters);
-
-class FNeuralPostProcessingBuildIndirectDispatchArgsCS : public FGlobalShader
+namespace NeuralPostProcessng
 {
-public:
-	DECLARE_GLOBAL_SHADER(FNeuralPostProcessingBuildIndirectDispatchArgsCS);
-	SHADER_USE_PARAMETER_STRUCT(FNeuralPostProcessingBuildIndirectDispatchArgsCS, FGlobalShader)
-
-	static void ModifyCompilationEnvironment(const FGlobalShaderPermutationParameters& Parameters, FShaderCompilerEnvironment& OutEnvironment);
 
-	BEGIN_SHADER_PARAMETER_STRUCT(FParameters, )
-		SHADER_PARAMETER_RDG_BUFFER_UAV(RWBuffer<uint>, RWIndirectDispatchArgsBuffer)
-		SHADER_PARAMETER_RDG_BUFFER_SRV(Buffer<uint>, SourceType)
-		SHADER_PARAMETER(FIntPoint, NetworkTextureSize)
-	END_SHADER_PARAMETER_STRUCT()
-
-	static bool ShouldCompilePermutation(const FGlobalShaderPermutationParameters& Parameters);
-};
-
-class FNeuralPostProcessingPrepareInputCS : public FGlobalShader
-{
-	DECLARE_GLOBAL_SHADER(FNeuralPostProcessingPrepareInputCS)
-	SHADER_USE_PARAMETER_STRUCT(FNeuralPostProcessingPrepareInputCS, FGlobalShader)
-
-public:
-
-	BEGIN_SHADER_PARAMETER_STRUCT(FParameters, )
-		SHADER_PARAMETER_STRUCT(FNueralPostProcessInput, Input0)
-		SHADER_PARAMETER_SAMPLER(SamplerState, InputTextureSampler)
-		RDG_BUFFER_ACCESS(IndirectDispatchArgsBuffer, ERHIAccess::IndirectArgs)
-		SHADER_PARAMETER(int32, InputBufferWidth)
-		SHADER_PARAMETER(int32, InputBufferHeight)
-		SHADER_PARAMETER_RDG_BUFFER_UAV(RWBuffer<float>, InputBuffer)
-		SHADER_PARAMETER(float, ColorScale)
-	END_SHADER_PARAMETER_STRUCT()
-
-	static void ModifyCompilationEnvironment(const FGlobalShaderPermutationParameters& InParameters, FShaderCompilerEnvironment& OutEnvironment);
-};
-
-class FNeuralPostProcessingProcessOutputPS : public FGlobalShader
-{
-	DECLARE_GLOBAL_SHADER(FNeuralPostProcessingProcessOutputPS)
-	SHADER_USE_PARAMETER_STRUCT(FNeuralPostProcessingProcessOutputPS, FGlobalShader)
-
-public:
-
-	BEGIN_SHADER_PARAMETER_STRUCT(FParameters, )
-		SHADER_PARAMETER_STRUCT(FNueralPostProcessInput, Input0)
-		SHADER_PARAMETER(int32, OutputBufferWidth)
-		SHADER_PARAMETER(int32, OutputBufferHeight)
-		SHADER_PARAMETER_RDG_BUFFER_UAV(RWBuffer<float>, OutputBuffer)
-		SHADER_PARAMETER(float, ColorScale)
-		RENDER_TARGET_BINDING_SLOTS()
-	END_SHADER_PARAMETER_STRUCT()
-
-	static void ModifyCompilationEnvironment(const FGlobalShaderPermutationParameters& InParameters, FShaderCompilerEnvironment& OutEnvironment);
-};
+	BEGIN_SHADER_PARAMETER_STRUCT(FNueralPostProcessInput, )
+		SHADER_PARAMETER_STRUCT_INCLUDE(FScreenPassTextureViewportParameters, Viewport)
+		SHADER_PARAMETER_RDG_TEXTURE(Texture2D, Texture)
+		END_SHADER_PARAMETER_STRUCT()
+
+	FNueralPostProcessInput GetNeuralPostProcessInput(FRDGTextureRef Texture, const FScreenPassTextureViewportParameters& ViewportParameters);
+
+	class FNeuralPostProcessingBuildIndirectDispatchArgsCS : public FGlobalShader
+	{
+	public:
+		DECLARE_GLOBAL_SHADER(FNeuralPostProcessingBuildIndirectDispatchArgsCS);
+		SHADER_USE_PARAMETER_STRUCT(FNeuralPostProcessingBuildIndirectDispatchArgsCS, FGlobalShader)
+
+		static void ModifyCompilationEnvironment(const FGlobalShaderPermutationParameters& Parameters, FShaderCompilerEnvironment& OutEnvironment);
+
+		BEGIN_SHADER_PARAMETER_STRUCT(FParameters, )
+			SHADER_PARAMETER_RDG_BUFFER_UAV(RWBuffer<uint>, RWIndirectDispatchArgsBuffer)
+			SHADER_PARAMETER_RDG_BUFFER_SRV(Buffer<uint>, SourceType)
+			SHADER_PARAMETER(FIntPoint, TargetDimension)
+		END_SHADER_PARAMETER_STRUCT()
+
+		static bool ShouldCompilePermutation(const FGlobalShaderPermutationParameters& Parameters);
+	};
+
+	class FDownScaleTextureCS : public FGlobalShader
+	{
+		DECLARE_GLOBAL_SHADER(FDownScaleTextureCS)
+		SHADER_USE_PARAMETER_STRUCT(FDownScaleTextureCS, FGlobalShader)
+
+	public:
+
+		BEGIN_SHADER_PARAMETER_STRUCT(FParameters, )
+			SHADER_PARAMETER_STRUCT(FNueralPostProcessInput, Source)
+			SHADER_PARAMETER_SAMPLER(SamplerState, SourceTextureSampler)
+			SHADER_PARAMETER(int32,	TargetWidth)
+			SHADER_PARAMETER(int32, TargetHeight)
+			SHADER_PARAMETER_RDG_TEXTURE_UAV(RWTexture2D<float4>, TargetTexture)
+			RDG_BUFFER_ACCESS(IndirectDispatchArgsBuffer, ERHIAccess::IndirectArgs)
+		END_SHADER_PARAMETER_STRUCT()
+
+		static void ModifyCompilationEnvironment(const FGlobalShaderPermutationParameters& InParameters, FShaderCompilerEnvironment& OutEnvironment);
+	};
+
+	class FDownScaleTexture : public FGlobalShader
+	{
+		DECLARE_GLOBAL_SHADER(FDownScaleTexture)
+		SHADER_USE_PARAMETER_STRUCT(FDownScaleTexture, FGlobalShader)
+
+	public:
+
+		BEGIN_SHADER_PARAMETER_STRUCT(FParameters, )
+			SHADER_PARAMETER_STRUCT(FNueralPostProcessInput, Source)
+			SHADER_PARAMETER_SAMPLER(SamplerState, SourceTextureSampler)
+			SHADER_PARAMETER(int32, TargetWidth)
+			SHADER_PARAMETER(int32, TargetHeight)
+			RENDER_TARGET_BINDING_SLOTS()
+		END_SHADER_PARAMETER_STRUCT()
+
+		static void ModifyCompilationEnvironment(const FGlobalShaderPermutationParameters& InParameters, FShaderCompilerEnvironment& OutEnvironment);
+	};
+
+	class FUpscaleTexture :public FGlobalShader
+	{
+		DECLARE_GLOBAL_SHADER(FUpscaleTexture)
+		SHADER_USE_PARAMETER_STRUCT(FUpscaleTexture, FGlobalShader)
+	public:
+		BEGIN_SHADER_PARAMETER_STRUCT(FParameters, )
+			SHADER_PARAMETER_RDG_TEXTURE(Texture2D<float4>, Source_Texture)
+			SHADER_PARAMETER(int32, SourceWidth)
+			SHADER_PARAMETER(int32, SourceHeight)
+			SHADER_PARAMETER(FIntPoint,ViewportSize)
+			RENDER_TARGET_BINDING_SLOTS()
+		END_SHADER_PARAMETER_STRUCT()
+
+		static void ModifyCompilationEnvironment(const FGlobalShaderPermutationParameters& InParameters, FShaderCompilerEnvironment& OutEnvironment);
+	};
+
+	class FCopyBetweenTextureAndOverlappedTileBufferCS : public FGlobalShader
+	{
+		DECLARE_GLOBAL_SHADER(FCopyBetweenTextureAndOverlappedTileBufferCS)
+		SHADER_USE_PARAMETER_STRUCT(FCopyBetweenTextureAndOverlappedTileBufferCS, FGlobalShader)
+
+	public:
+
+		BEGIN_SHADER_PARAMETER_STRUCT(FParameters, )
+			SHADER_PARAMETER(int32, SourceWidth)
+			SHADER_PARAMETER(int32, SourceHeight)
+			SHADER_PARAMETER_RDG_TEXTURE_UAV(RWTexture2D<float4>,RWSourceTexture)
+			SHADER_PARAMETER_SAMPLER(SamplerState, SourceTextureSampler)
+
+			SHADER_PARAMETER(int32, TargetOverlappedTileWidth)
+			SHADER_PARAMETER(int32, TargetOverlappedTileHeight)
+			SHADER_PARAMETER(FIntPoint, ViewTileDimension)
+			SHADER_PARAMETER(FVector2f, TileOverlap)
+			SHADER_PARAMETER(int32, NumOfChannel)
+			SHADER_PARAMETER(int32, bVisualizeOverlap)
+			SHADER_PARAMETER(float, OverlapVisualizeIntensity)
+			SHADER_PARAMETER_RDG_BUFFER_UAV(RWBuffer<float4>, TargetBuffer)
+
+			RDG_BUFFER_ACCESS(IndirectDispatchArgsBuffer, ERHIAccess::IndirectArgs)
+		END_SHADER_PARAMETER_STRUCT()
+
+		static void ModifyCompilationEnvironment(const FGlobalShaderPermutationParameters& InParameters, FShaderCompilerEnvironment& OutEnvironment);
+
+		enum class EDirection : uint32
+		{
+			ToOverlappedTiles,
+			FromOverlappedTiles,
+			MAX
+		};
+
+		enum class EOverlapResolveType : uint32
+		{
+			Ignore,
+			Feathering,
+			MAX
+		};
+
+		class FDimensionCopyDirection : SHADER_PERMUTATION_ENUM_CLASS("BUFFER_COPY_DIRECTION", EDirection);
+		class FDimensionOverlapResolveType : SHADER_PERMUTATION_ENUM_CLASS("OVERLAP_RESOLVE_TYPE", EOverlapResolveType);
+		using FPermutationDomain = TShaderPermutationDomain<FDimensionCopyDirection, FDimensionOverlapResolveType>;
+	};
+}
\ No newline at end of file
diff --git a/Engine/Plugins/Experimental/NeuralRendering/Source/NeuralPostProcessing/Shaders/NeuralPostProcessing.usf b/Engine/Plugins/Experimental/NeuralRendering/Source/NeuralPostProcessing/Shaders/NeuralPostProcessing.usf
--- a/Engine/Plugins/Experimental/NeuralRendering/Source/NeuralPostProcessing/Shaders/NeuralPostProcessing.usf
+++ b/Engine/Plugins/Experimental/NeuralRendering/Source/NeuralPostProcessing/Shaders/NeuralPostProcessing.usf
@@ -7,18 +7,43 @@
 #define SOURCE_TYPE_TEXTURE		0
 #define SOURCE_TYPE_BUFFER		1
 
-Texture2D Input0_Texture;
-SCREEN_PASS_TEXTURE_VIEWPORT(Input0)
+#define BUFFER_COPY_TO_OVERLAPPING_TILE		0
+#define BUFFER_COPY_FROM_OVERLAPPING_TILE	1
 
-float2 ConvertGridPos2UV(uint2 GridPosition)
+#ifndef BUFFER_COPY_DIRECTION
+#define BUFFER_COPY_DIRECTION				BUFFER_COPY_TO_OVERLAPPING_TILE
+#endif
+
+#define OVERLAP_RESOLVE_IGNORE				0
+#define OVERLAP_RESOLVE_FEATHERING			1
+
+#ifndef OVERLAP_RESOLVE_TYPE
+#define OVERLAP_RESOLVE_TYPE OVERLAP_RESOLVE_IGNORE
+#endif
+
+float2 ConvertGridPos2UV(uint2 GridPosition,float2 ExtentInverse)
 {
 	float2 GripPositionF = float2(GridPosition);
-	return Input0_ExtentInverse * (GripPositionF + 0.5f);
+	return ExtentInverse * (GripPositionF + 0.5f);
 }
 
 Buffer<uint> SourceType;
 RWBuffer<uint> RWIndirectDispatchArgsBuffer;
-int2 NetworkTextureSize;
+int2 TargetDimension;
+
+int SourceWidth;
+int SourceHeight;
+int2 ViewportSize;
+Texture2D<float4> Source_Texture;
+SCREEN_PASS_TEXTURE_VIEWPORT(Source)
+SamplerState SourceTextureSampler;
+
+int TargetWidth;
+int TargetHeight;
+RWTexture2D<float4> TargetTexture;
+
+int bVisualizeOverlap;
+float OverlapVisualizeIntensity;
 
 [numthreads(8, 1, 1)]
 void BuildIndirectDispatchArgsCS(uint2 DispatchThreadId : SV_DispatchThreadID)
@@ -27,10 +52,11 @@ void BuildIndirectDispatchArgsCS(uint2 DispatchThreadId : SV_DispatchThreadID)
 	{
 		if (SourceType[0] == SOURCE_TYPE_TEXTURE)
 		{
-			uint IndirectDispatchX = (NetworkTextureSize.x + THREAD_GROUP_SIZE - 1) / THREAD_GROUP_SIZE;
-			uint IndirectDispatchY = (NetworkTextureSize.y + THREAD_GROUP_SIZE - 1) / THREAD_GROUP_SIZE;
+			uint IndirectDispatchX = (TargetDimension.x + THREAD_GROUP_SIZE - 1) / THREAD_GROUP_SIZE;
+			uint IndirectDispatchY = (TargetDimension.y + THREAD_GROUP_SIZE - 1) / THREAD_GROUP_SIZE;
+			uint IndirectDispatchZ = 1;
 
-			WriteDispatchIndirectArgs(RWIndirectDispatchArgsBuffer, 0, IndirectDispatchX, IndirectDispatchY, 1);
+			WriteDispatchIndirectArgs(RWIndirectDispatchArgsBuffer, 0, IndirectDispatchX, IndirectDispatchY, IndirectDispatchZ);
 		}
 		else
 		{
@@ -39,26 +65,55 @@ void BuildIndirectDispatchArgsCS(uint2 DispatchThreadId : SV_DispatchThreadID)
 	}
 }
 
-SamplerState InputTextureSampler;
-int InputBufferWidth;
-int InputBufferHeight;
-RWBuffer<float> InputBuffer;
-float ColorScale;
+float3 BilinearInterpolate(float3 V00, float3 V10, float3 V01, float3 V11, float2 LerpFactors)
+{
+	float Alpha = LerpFactors.x;
+	float Beta = LerpFactors.y;
 
-// from buffer to neural input
-[numthreads(THREAD_GROUP_SIZE, THREAD_GROUP_SIZE, 1)]
-void PrepareInput(in const uint3 DispatchThreadID : SV_DispatchThreadID)
+	return (1.0 - Beta) * ((1.0 - Alpha) * V00 + Alpha * V10) + Beta * ((1.0 - Alpha) * V01 + Alpha * V11);
+}
+
+float4 BilinearInterpolate(float4 V00, float4 V10, float4 V01, float4 V11, float2 LerpFactors)
 {
-	if (DispatchThreadID.x >= InputBufferWidth || DispatchThreadID.y >= InputBufferHeight)
-	{
-		return;
-	}
+	float Alpha = LerpFactors.x;
+	float Beta = LerpFactors.y;
+
+	return (1.0 - Beta) * ((1.0 - Alpha) * V00 + Alpha * V10) + Beta * ((1.0 - Alpha) * V01 + Alpha * V11);
+}
 
-	// calculate the grid size mapping from the rendered output to the input of the neuralnetwork
-	const float2 BufferSize = float2(InputBufferWidth, InputBufferHeight) - 0.5f;
-	const float2 WidthFactors = Input0_ViewportSize.xy / BufferSize;
-	const int2 LoXY = Input0_ViewportMin.xy + (int2)(WidthFactors * (float2)DispatchThreadID.xy);
-	const int2 HiXY = Input0_ViewportMin.xy + (int2)ceil(WidthFactors * (float2)(DispatchThreadID.xy + 1));
+float4 UpscaleTexture(float4 Position : SV_POSITION) : SV_Target0
+{
+	const int2 BufferSize = int2(SourceWidth, SourceHeight);
+	float X = Position.x * float(BufferSize.x - 1) / (float)ViewportSize.x;
+	float Y = Position.y * float(BufferSize.y - 1) / (float)ViewportSize.y;
+
+	int LoX = (int)clamp(floor(X), 0.0, float(BufferSize.x - 1));
+	int HiX = (int)clamp(floor(X + 1.0), 0.0, float(BufferSize.x - 1));
+
+	int LoY = (int)clamp(floor(Y), 0.0, float(BufferSize.y - 1));
+	int HiY = (int)clamp(floor(Y + 1.0), 0.0, float(BufferSize.y - 1));
+
+	// bilinear filtering
+	float3 A = Source_Texture.Load(int3(LoX, LoY, 0)).xyz;
+	float3 B = Source_Texture.Load(int3(HiX, LoY, 0)).xyz;
+	float3 C = Source_Texture.Load(int3(LoX, HiY, 0)).xyz;
+	float3 D = Source_Texture.Load(int3(HiX, HiY, 0)).xyz;
+
+	float Alpha = clamp(X - (float)LoX, 0.0, 1.0);
+	float Beta = clamp(Y - (float)LoY, 0.0, 1.0);
+
+	//float3 Result = (1.0 - Beta) * ((1.0 - Alpha) * A + Alpha * B) + Beta * ((1.0 - Alpha) * C + Alpha * D);
+	float3 Result = BilinearInterpolate(A, B, C, D, float2(Alpha, Beta));
+
+	return float4(Result, 1.0);
+}
+
+float4 GetSourceGridAverageValue(float2 Position)
+{
+	const float2 BufferSize = float2(TargetWidth, TargetHeight);
+	const float2 WidthFactors = Source_ViewportSize.xy / BufferSize;
+	const int2 LoXY = Source_ViewportMin.xy + (int2)(WidthFactors * Position);
+	const int2 HiXY = Source_ViewportMin.xy + (int2)ceil(WidthFactors * (Position + 1.0f));
 	int LoX = LoXY.x;
 	int LoY = LoXY.y;
 	int HiX = HiXY.x;
@@ -69,72 +124,392 @@ void PrepareInput(in const uint3 DispatchThreadID : SV_DispatchThreadID)
 	float Div = 0.0;
 	for (int x = LoX; x <= HiX; x++)
 	{
-		if (x < Input0_Extent.x)
+		if (x < Source_Extent.x)
 		{
 			for (int y = LoY; y <= HiY; y++)
 			{
-				if (y < Input0_Extent.y)
+				if (y < Source_Extent.y)
 				{
-					float2 BufferUV = ConvertGridPos2UV(uint2(x,y));
-					float2 ClampedBufferUV = clamp(BufferUV, Input0_UVViewportBilinearMin, Input0_UVViewportBilinearMax);
-					Result += Input0_Texture.SampleLevel(InputTextureSampler, ClampedBufferUV, 0);
+					float2 BufferUV = ConvertGridPos2UV(uint2(x, y), Source_ExtentInverse);
+					float2 ClampedBufferUV = clamp(BufferUV, Source_UVViewportBilinearMin, Source_UVViewportBilinearMax);
+					Result += Source_Texture.SampleLevel(SourceTextureSampler, ClampedBufferUV, 0);
 					Div += 1.0;
 				}
 			}
 		}
 	}
+
 	Result /= Div;
 
-	// write to the input buffer of the neural network
-	int Idx = DispatchThreadID.y * InputBufferWidth + DispatchThreadID.x;
-	int Offset1 = InputBufferWidth * InputBufferHeight;
-	int Offset2 = Offset1 + Offset1;
+	return Result;
+}
+
+float4 DownScaleTexture(float4 Position : SV_POSITION) : SV_Target0
+{
+	float4 Result = GetSourceGridAverageValue(Position.xy);
+	return Result;
+}
+
+[numthreads(THREAD_GROUP_SIZE, THREAD_GROUP_SIZE, 1)]
+void DownScaleTextureCS(in const uint3 DispatchThreadID : SV_DispatchThreadID)
+{
+	if (DispatchThreadID.x >= TargetWidth || DispatchThreadID.y >= TargetHeight)
+	{
+		return;
+	}
+
+	float4 Result = GetSourceGridAverageValue(DispatchThreadID.xy);
+	TargetTexture[DispatchThreadID.xy] = Result;
+}
 
-	InputBuffer[Idx] = Result.r;
-	InputBuffer[Idx + Offset1] = Result.g;
-	InputBuffer[Idx + Offset2] = Result.b;
+bool IsInRect(int2 Position, int2 MinPosition, int2 MaxPosition)
+{
+	return all(Position >= MinPosition) && all(Position < MaxPosition);
 }
 
-int OutputBufferWidth;
-int OutputBufferHeight;
-RWBuffer<float> OutputBuffer;
+RWTexture2D<float4> RWSourceTexture;
+int TargetOverlappedTileWidth;
+int TargetOverlappedTileHeight;
+float2 TileOverlap;
+RWBuffer<float> TargetBuffer;
+int2 ViewTileDimension;
+int NumOfChannel;
+
+#define TPVector2	float2
 
-// from neural output to buffer, full screenpass pixel shader
-float4 ProcessOutput(float4 Position : SV_POSITION) : SV_Target0
+struct FTiledPosition
 {
-	float X = Position.x * float(OutputBufferWidth - 1) / (float)Input0_ViewportSize.x;
-	float Y = Position.y * float(OutputBufferHeight - 1) / (float)Input0_ViewportSize.y;
+	int2 Index;		
+	TPVector2 Position;
 
-	int LoX =  (int)clamp(floor(X), 0.0, float(OutputBufferWidth - 1));
-	int HiX = (int)clamp(floor(X + 1.0), 0.0, float(OutputBufferWidth - 1));
+	// For a given Tiled Position, get the corresponding Position with respect to tile dimension and overlap
+	// @TODO: float2 overlap
+	TPVector2 ToPosition(int2 TileDimension, TPVector2 Overlaps)
+	{
+		const TPVector2 SourceTextureOffset = Index * TileDimension - Overlaps;
+		return SourceTextureOffset + Position;
+	}
 
-	int LoY = (int)clamp(floor(Y), 0.0, float(OutputBufferHeight - 1));
-	int HiY = (int)clamp(floor(Y + 1.0), 0.0, float(OutputBufferHeight - 1));
+	// Transform the current tiled position to the corresponding tiled position in the neighbor tile.
+	FTiledPosition TransformBy(int2 Offset, int2 TileDimension, TPVector2 Overlaps)
+	{
+		FTiledPosition Neighbor = (FTiledPosition)0;
+		Neighbor.Index = Index + Offset;
+		TPVector2 AbsoluteNeighborPosition = Neighbor.ToPosition(TileDimension, Overlaps);
+		TPVector2 AbsolutePosition = ToPosition(TileDimension, Overlaps);
+		Neighbor.Position = AbsolutePosition - AbsoluteNeighborPosition;
 
-	int Offset1 = OutputBufferWidth * OutputBufferHeight;
-	int Offset2 = Offset1 + Offset1;
+		return Neighbor;
+	}
 
-	// bilinear filtering
-	int Idx;
+	bool IsLeftEdge()
+	{
+		return Index.x == 0;
+	}
 
-	Idx = LoY * OutputBufferWidth + LoX;
-	float3 A = float3(OutputBuffer[Idx], OutputBuffer[Idx + Offset1], OutputBuffer[Idx + Offset2]);
+	bool IsRightEdge()
+	{
+		return Index.x == (ViewTileDimension.x - 1);
+	}
 
-	Idx = LoY * OutputBufferWidth + HiX;
-	float3 B = float3(OutputBuffer[Idx], OutputBuffer[Idx + Offset1], OutputBuffer[Idx + Offset2]);
+	bool IsTopEdge()
+	{
+		return Index.y == 0;
+	}
 
-	Idx = HiY * OutputBufferWidth + LoX;
-	float3 C = float3(OutputBuffer[Idx], OutputBuffer[Idx + Offset1], OutputBuffer[Idx + Offset2]);
+	bool IsBottomEdge()
+	{
+		return Index.y == (ViewTileDimension.y - 1);
+	}
+};
 
-	Idx = HiY * OutputBufferWidth + HiX;
-	float3 D = float3(OutputBuffer[Idx], OutputBuffer[Idx + Offset1], OutputBuffer[Idx + Offset2]);
+FTiledPosition GetTiledPosition(int2 Position, int2 TileDimension)
+{
+	FTiledPosition TiledPosition = (FTiledPosition)0;
+	TiledPosition.Index = Position / TileDimension;
+	TiledPosition.Position = Position % TileDimension;
 
-	float Alpha = clamp(X - (float)LoX, 0.0, 1.0);
-	float Beta = clamp(Y - (float)LoY, 0.0, 1.0);
+	return TiledPosition;
+}
+
+float4 SampleSourceTexture(TPVector2 SamplePosition,int2 TextureSize)
+{
+#if 0
+	const float2 ExtentInverse = 1.0f/float2(SourceWidth, SourceHeight);
+	float2  UV = ConvertGridPos2UV(SamplePosition, ExtentInverse);
+	return  RWSourceTexture.SampleLevel(SourceTextureSampler, UV, 0);
+#endif
+	TPVector2 TextureSizeMax = TextureSize -1;
+	TPVector2 MirrorPosition = abs(TextureSizeMax - abs(SamplePosition - TextureSizeMax));
+	/*if (any(SamplePosition >= TextureSize) || any(SamplePosition < 0))
+	{
+		return 0.0f;
+	}*/
+	return RWSourceTexture[(uint2)MirrorPosition];
+}
 
-	float3 Result = (1.0 - Beta) * ((1.0 - Alpha) * A + Alpha * B) + Beta * ((1.0 - Alpha) * C + Alpha * D);
+struct FBufferLocation
+{
+	int Offset;
+	int ChannelSize;
+	int Valid;
+};
 
-	Result = Result.rgb;
+FBufferLocation GetBufferLocation(FTiledPosition TiledPosition,int2 TargetOverlappedTileDim)
+{
+	FBufferLocation BufferLocation = (FBufferLocation)0;
+	BufferLocation.Valid = all(TiledPosition.Position >= 0) && 
+		all(TiledPosition.Position < TargetOverlappedTileDim)&&
+		all(TiledPosition.Index >= 0) &&
+		all(TiledPosition.Index< ViewTileDimension);
 
-	return float4(Result, 1.0);
-}
\ No newline at end of file
+	if (BufferLocation.Valid)
+	{
+		const int ChannelSize = TargetOverlappedTileDim.x * TargetOverlappedTileDim.y;
+		const int BatchSize = ChannelSize * NumOfChannel;
+
+		const int BatchIndex = TiledPosition.Index.x + TiledPosition.Index.y * ViewTileDimension.x;
+		const int BatchOffset = BatchIndex * BatchSize;
+		const int Offset = BatchOffset + (uint(TiledPosition.Position.x) + uint(TiledPosition.Position.y) * TargetOverlappedTileDim.x);
+
+
+		BufferLocation.ChannelSize = ChannelSize;
+		BufferLocation.Offset = Offset;
+	}
+
+	return BufferLocation;
+}
+
+
+void SetTargetBufferValue(FBufferLocation BufferLocation, float4 Value)
+{
+	for (int i = 0; i < NumOfChannel; ++i)
+	{
+		TargetBuffer[BufferLocation.Offset + i * BufferLocation.ChannelSize] = Value[i];
+	}
+}
+
+// Use the default value if the Buffer location is not valid
+float4 GetTargetBufferValue(FBufferLocation BufferLocation, float4 DefaultValue)
+{
+	float4 Value = DefaultValue;
+
+	if (BufferLocation.Valid)
+	{
+		for (int i = 0; i < NumOfChannel; ++i)
+		{
+			Value[i] = TargetBuffer[BufferLocation.Offset + i * BufferLocation.ChannelSize];
+		}
+	}
+
+	return Value;
+}
+
+float4 GetValueFromTransformedTile(FTiledPosition TiledPosition, int2 Offset, 
+								   int2 TileDimension, TPVector2 InOverlap, float4 DefaultValue)
+{
+	FTiledPosition TransformedPosition = TiledPosition.TransformBy(Offset, TileDimension, InOverlap);
+	FBufferLocation TransformedBufferLocation = GetBufferLocation(TransformedPosition, TileDimension + 2 * InOverlap);
+	return GetTargetBufferValue(TransformedBufferLocation, DefaultValue);
+}
+
+float4 EmbedDebugValue(float4 Value, float4 DebugValue, bool ForceDebug)
+{
+	//@todo: compile time dynamic for debug view.
+	return (bVisualizeOverlap||ForceDebug) ? (OverlapVisualizeIntensity*DebugValue) : Value;
+}
+
+
+[numthreads(THREAD_GROUP_SIZE, THREAD_GROUP_SIZE, 1)]
+void CopyBetweenTextureAndOverlappedTileBufferCS(in const uint3 DispatchThreadID : SV_DispatchThreadID)
+{
+	const int2 TargetOverlappedTileDim = int2(TargetOverlappedTileWidth, TargetOverlappedTileHeight);
+	const int2 TargetBufferDim = TargetOverlappedTileDim * ViewTileDimension;
+	const int2 SourceTileDim = TargetOverlappedTileDim - 2 * TileOverlap;
+	const int2 SourceBufferDim = min(SourceTileDim * ViewTileDimension, int2(SourceWidth, SourceHeight));
+
+//#if BUFFER_COPY_DIRECTION == BUFFER_COPY_TO_OVERLAPPING_TILE
+	if (any(DispatchThreadID.xy >= TargetBufferDim))
+//#elif  BUFFER_COPY_DIRECTION == BUFFER_COPY_FROM_OVERLAPPING_TILE
+//	if (any(DispatchThreadID.xy >= SourceBufferDim))
+//#endif
+	{
+		return;
+	}
+
+	// Derive source texture position and buffer position
+	const FTiledPosition TiledPosition = GetTiledPosition(DispatchThreadID.xy, TargetOverlappedTileDim);
+	const TPVector2 TexturePosition = TiledPosition.ToPosition(SourceTileDim, TileOverlap);
+	const FBufferLocation BufferLocation = GetBufferLocation(TiledPosition, TargetOverlappedTileDim);
+
+#if BUFFER_COPY_DIRECTION == BUFFER_COPY_TO_OVERLAPPING_TILE
+
+	float4 Value = SampleSourceTexture(TexturePosition, SourceBufferDim);
+	SetTargetBufferValue(BufferLocation, Value);
+	
+#elif BUFFER_COPY_DIRECTION == BUFFER_COPY_FROM_OVERLAPPING_TILE
+	// Read back from tiled buffer to buffer
+	const uint2 WriteTexturePosition = (uint2)TexturePosition;
+
+	// write pixel position is outside of the texture.
+	if (any(WriteTexturePosition >= SourceBufferDim))
+	{
+		return;
+	}
+
+	const float4 Value0 = GetTargetBufferValue(BufferLocation, 0.0f);
+	#if OVERLAP_RESOLVE_TYPE == OVERLAP_RESOLVE_IGNORE
+		if (IsInRect(TiledPosition.Position, TileOverlap, TargetOverlappedTileDim - TileOverlap))
+		{
+			// Internal region with no overlapping with other tiles
+			RWSourceTexture[WriteTexturePosition] = Value0;
+		}
+	#elif OVERLAP_RESOLVE_TYPE == OVERLAP_RESOLVE_FEATHERING
+		if (IsInRect(TiledPosition.Position, 2*TileOverlap, TargetOverlappedTileDim - 2 * TileOverlap))
+		{
+			// Internal region with no overlapping with other tiles
+			RWSourceTexture[WriteTexturePosition] = Value0;
+		}
+		else if (IsInRect(TiledPosition.Position, TileOverlap, TargetOverlappedTileDim - TileOverlap))
+		{
+			//Linear blending overlapping regions with other tiles. 
+			if (IsInRect(TiledPosition.Position, TileOverlap + int2(0, TileOverlap.y), int2(2 * TileOverlap.x, TargetOverlappedTileDim.y - 2 * TileOverlap.y)))
+			{
+				//Section 1: lerping value in tile (x,y) and tile(x-1,y)
+				float4 Value1 = GetValueFromTransformedTile(TiledPosition, int2(-1, 0), SourceTileDim, TileOverlap, Value0);
+
+				float LerpFactor = 1 - TiledPosition.Position.x / (2.0f * TileOverlap.x);
+				RWSourceTexture[WriteTexturePosition] = EmbedDebugValue(lerp(Value0, Value1, LerpFactor), 1.2f, false);
+			}
+			else if (IsInRect(TiledPosition.Position, 
+				int2(TargetOverlappedTileDim.x - 2 * TileOverlap.x, 2 * TileOverlap.y),
+				int2(TargetOverlappedTileDim.x - TileOverlap.x, TargetOverlappedTileDim.y - 2 * TileOverlap.y)))
+			{
+				//Section 2: lerping value in tile (x,y) and tile (x+1,y)
+				float4 Value1 = GetValueFromTransformedTile(TiledPosition, int2(1, 0), SourceTileDim, TileOverlap, Value0);
+
+				float LerpFactor = 1 - (TargetOverlappedTileDim.x - TiledPosition.Position.x) / (2.0f * TileOverlap.x);
+				RWSourceTexture[WriteTexturePosition] = EmbedDebugValue(lerp(Value0, Value1, LerpFactor), float4(1.5, 0.0f, 0.0f, 1.0f), false);
+
+			}
+			else if (IsInRect(TiledPosition.Position,
+				int2(2 * TileOverlap.x, TileOverlap.y),
+				int2(TargetOverlappedTileDim.x - 2 * TileOverlap.x, 2 * TileOverlap.y)))
+			{
+				//Section 3: lerping value in tile (x,y) and tile (x,y-1)
+				float4 Value1 = GetValueFromTransformedTile(TiledPosition, int2(0, -1), SourceTileDim, TileOverlap, Value0);
+
+				float LerpFactor = 1 - TiledPosition.Position.y / (2.0f * TileOverlap.y);
+				RWSourceTexture[WriteTexturePosition] = EmbedDebugValue(lerp(Value0, Value1, LerpFactor), float4(0.0, 0.8f, 0.0f, 1.0f), false);
+			}
+			else if (IsInRect(TiledPosition.Position,
+				int2(2 * TileOverlap.x, TargetOverlappedTileDim.y - 2 * TileOverlap.y),
+				int2(TargetOverlappedTileDim.x - 2 * TileOverlap.x, TargetOverlappedTileDim.y - TileOverlap.y)))
+			{
+				//Section 4: lerping value in tile (x,y) and tile (x,y+1)
+				float4 Value1 = GetValueFromTransformedTile(TiledPosition, int2(0, 1), SourceTileDim, TileOverlap, Value0);
+				
+				float LerpFactor = 1 - (TargetOverlappedTileDim.y - TiledPosition.Position.y) / (2.0f * TileOverlap.y);
+				RWSourceTexture[WriteTexturePosition] = EmbedDebugValue(lerp(Value0, Value1, LerpFactor), float4(0.0, 0.2f, 0.9f, 1.0f), false);
+			}
+			else if (IsInRect(TiledPosition.Position,
+				TileOverlap, 2 * TileOverlap))
+			{
+				//Section 5: lerping value in tile (x,y) and tile (x-1,y-1),(x,y-1),(x-1,y)
+				// V00(-1, -1), V10(0, -1)
+				// V01(-1,  0), V11(0,  0)
+				// Assume horizontal edge
+				float4 V11 = Value0;
+				float4 V01 = GetValueFromTransformedTile(TiledPosition, int2(-1,  0), SourceTileDim, TileOverlap, Value0);
+				float4 V00 = GetValueFromTransformedTile(TiledPosition, int2(-1, -1), SourceTileDim, TileOverlap,    V01);
+				float4 V10 = GetValueFromTransformedTile(TiledPosition, int2( 0, -1), SourceTileDim, TileOverlap, Value0);
+				// Check left edge
+				if (TiledPosition.IsLeftEdge())
+				{
+					V00 = V10;
+				}
+
+				float2 LerpFactors = TiledPosition.Position/(2.0f * TileOverlap);
+				// bilinear interpolation
+				float4 FilterdValue = BilinearInterpolate(V00, V10, V01, V11, LerpFactors);
+				RWSourceTexture[WriteTexturePosition] = EmbedDebugValue(FilterdValue, float4(1.0, 1.2f, 0.2f, 1.0f), false);
+			}
+			else if (IsInRect(TiledPosition.Position,
+				int2(TargetOverlappedTileDim.x-3*TileOverlap.x,0)+ TileOverlap, 
+				int2(TargetOverlappedTileDim.x - 3 * TileOverlap.x, 0) + 2 * TileOverlap))
+			{
+				//Section 6: lerping value in tile (x,y) and tile (x+1,y-1),(x+1,y),(x,y-1)
+				// V00(0, -1), V10(1, -1)
+				// V01(0,  0), V11(1,  0)
+				// Assume horizontal edge
+				float4 V00 = GetValueFromTransformedTile(TiledPosition, int2( 0, -1), SourceTileDim, TileOverlap, Value0);
+				float4 V11 = GetValueFromTransformedTile(TiledPosition, int2( 1,  0), SourceTileDim, TileOverlap, Value0);
+				float4 V01 = Value0;
+				float4 V10 = GetValueFromTransformedTile(TiledPosition, int2(1, -1), SourceTileDim, TileOverlap,     V11);
+				// Check right edge
+				if (TiledPosition.IsRightEdge())
+				{
+					V10 = V00;
+				}
+
+				float2 LerpFactors = (TiledPosition.Position - int2(SourceTileDim.x,0)) / (2.0f * TileOverlap);
+
+				// bilinear interpolation
+				float4 FilterdValue = BilinearInterpolate(V00, V10, V01, V11, LerpFactors);
+				RWSourceTexture[WriteTexturePosition] = EmbedDebugValue(FilterdValue, float4(2.0, 0.2f, 0.9f, 1.0f), false);
+			}
+			else if (IsInRect(TiledPosition.Position,
+				int2(TileOverlap.x, TargetOverlappedTileDim.y - 2 * TileOverlap.y),
+				int2(2 * TileOverlap.x, TargetOverlappedTileDim.y - TileOverlap.y)))
+			{
+				//Section 7: lerping value in tile (x,y) and tile (x-1,y),(x-1,y+1),(x,y+1)
+				// V00(-1,  0), V10(0, 0)
+				// V01(-1,  1), V11(0, 1)
+				// Assume horizontal edge
+				float4 V00 = GetValueFromTransformedTile(TiledPosition, int2(-1, 0), SourceTileDim, TileOverlap, Value0);
+				float4 V10 = Value0;
+				float4 V01 = GetValueFromTransformedTile(TiledPosition, int2(-1, 1), SourceTileDim, TileOverlap,    V00);
+				float4 V11 = GetValueFromTransformedTile(TiledPosition, int2( 0, 1), SourceTileDim, TileOverlap, Value0);
+				float2 LerpFactors = (TiledPosition.Position - int2(0, SourceTileDim.y)) / (2.0f * TileOverlap);
+				// Check left edge
+				if (TiledPosition.IsLeftEdge())
+				{
+					V01 = V11;
+				}
+
+				// bilinear interpolation
+				float4 FilterdValue = BilinearInterpolate(V00, V10, V01, V11, LerpFactors);
+				RWSourceTexture[WriteTexturePosition] = EmbedDebugValue(FilterdValue, float4(0.0, 2.2f, 2.9f, 1.0f), false);
+			}
+			else if (IsInRect(TiledPosition.Position,
+				int2(TargetOverlappedTileDim.x -2 * TileOverlap.x, TargetOverlappedTileDim.y - 2 * TileOverlap.y),
+				int2(TargetOverlappedTileDim.x - TileOverlap.x, TargetOverlappedTileDim.y - TileOverlap.y)))
+			{
+				//Section 8: lerping value in tile (x,y) and tile (x,y+1),(x+1,y),(x+1,y+1)
+				// V00( 0,  0), V10( 1, 0)
+				// V01( 0,  1), V11( 1, 1)
+				// Assume horizontal edge
+				float4 V00 = Value0;
+				float4 V10 = GetValueFromTransformedTile(TiledPosition, int2( 1, 0), SourceTileDim, TileOverlap, Value0);
+				float4 V01 = GetValueFromTransformedTile(TiledPosition, int2( 0, 1), SourceTileDim, TileOverlap,    V00);
+				float4 V11 = GetValueFromTransformedTile(TiledPosition, int2( 1, 1), SourceTileDim, TileOverlap,    V10);
+				float2 LerpFactors = (TiledPosition.Position - SourceTileDim) / (2.0f * TileOverlap);
+				// Check right edge
+				if (TiledPosition.IsRightEdge())
+				{
+					V11 = V01;
+				}
+
+				// bilinear interpolation
+				float4 FilterdValue = BilinearInterpolate(V00, V10, V01, V11, LerpFactors);
+				RWSourceTexture[WriteTexturePosition] = EmbedDebugValue(FilterdValue, float4(2.0, 2.2f, 0.1f, 1.0f), false);
+			}
+		}
+		else
+		{
+			// border regions that should be ignored
+		}
+	#endif// End of OVERLAP_RESOLVE_TYPE
+#endif
+
+}
diff --git a/Engine/Source/Runtime/Engine/Classes/Engine/NeuralProfile.h b/Engine/Source/Runtime/Engine/Classes/Engine/NeuralProfile.h
--- a/Engine/Source/Runtime/Engine/Classes/Engine/NeuralProfile.h
+++ b/Engine/Source/Runtime/Engine/Classes/Engine/NeuralProfile.h
@@ -34,7 +34,25 @@ if the input texture has different dimensions, it will be scaled down before app
 	e.g., if the model input dimension is (1x3x200x200) and the used buffer size of the post processing
 	is 1000x1000, then 5x5 tiles ((5x5)x3x200x200) will be run and recombined.
 	*/
-	//Auto			UMETA(DisplayName = "Auto")
+	Auto			UMETA(DisplayName = "Auto")
+};
+
+UENUM(BlueprintType)
+enum class ETileOverlapResolveType : uint8
+{
+	/** Overlapped tile regions have no contribution to adjecent tiles*/
+	Ignore			UMETA(DisplayName = "Ignore"),
+
+	/** Overlapped regions are blended linearly to adjecent tiles*/
+	Feathering		UMETA(DisplayName = "Feathering")
+};
+
+
+UENUM(BlueprintType)
+enum class ENeuralProfileRuntimeType : uint8
+{
+	NNERuntimeRDGDml UMETA(DisplayName = "NNERuntimeRDGDml"),
+	MAX				 UMETA(Hidden)
 };
 
 // struct with all the settings we want in UNeuralProfile, separate to make it easer to pass this data around in the engine.
@@ -57,9 +75,9 @@ struct FNeuralProfileStruct
 	UPROPERTY(Category = "Common", EditAnywhere, BlueprintReadOnly, meta = (DisplayName = "Output Format", editcondition = "false", EditConditionHides))
 	ENeuralProfileFormat OutputFormat;
 
-	//runtime name (support "NNERuntimeRDGDml" only at this moment)
+	//runtime type (support "NNERuntimeRDGDml" only at this moment)
 	UPROPERTY(Category = "Model", EditAnywhere, BlueprintReadOnly)
-	FString RuntimeName;
+	ENeuralProfileRuntimeType RuntimeType;
 
 	/** Stores the NNEModelData imported from e.g., onnx model */
 	UPROPERTY(EditAnywhere, BlueprintReadOnly, Category = Model, meta = (DisplayName = "NNE Model Data"))
@@ -73,24 +91,33 @@ struct FNeuralProfileStruct
 	UPROPERTY(EditAnywhere, BlueprintReadOnly, Category = Model, meta = (DisplayName = "Output Dimension", editcondition = "false"))
 	FIntVector4 OutputDimension;
 
-	/** Except for 1x1, all other model requires an intermediate buffer.*/
-	UPROPERTY(EditAnywhere, BlueprintReadOnly, Category = "Quality and Performance")
-	ENeuralModelTileType TileSize;
-
 	/** Used to override the batch size if the batch dimension is dynamic (-1)*/
-	UPROPERTY(EditAnywhere, BlueprintReadOnly, Category = "Quality and Performance")
+	UPROPERTY(EditAnywhere, BlueprintReadOnly, Category = "Override", meta = (DisplayName = "Batch Size"))
 	int32 BatchSizeOverride;
 
+	/** Total tiles used. Each tile will be executed by 1 batch */
+	UPROPERTY(EditAnywhere, Category = "Tile")
+	ENeuralModelTileType TileSize;
+
+	/** Tile border overlaps (Left|Right, Top|Bottom). The larger this value, the more tiles are required to cover the whole screen when TileSize is Auto.*/
+	UPROPERTY(EditAnywhere, Category = "Tile", meta = (DisplayName = "Border Overlaps"))
+	FIntPoint TileOverlap;
+
+	UPROPERTY(EditAnywhere, Category = "Tile", meta = (DisplayName = "Overlap Resolve Type"))
+	ETileOverlapResolveType TileOverlapResolveType;
+
 	FNeuralProfileStruct()
 	{
 		InputFormat = ENeuralProfileFormat::Type32;
 		OutputFormat = ENeuralProfileFormat::Type32;
-		RuntimeName = "NNERuntimeRDGDml";
+		RuntimeType = ENeuralProfileRuntimeType::NNERuntimeRDGDml;
 		NNEModelData = nullptr;
 		InputDimension = FIntVector4(0);
 		OutputDimension = FIntVector4(0);
 		TileSize = ENeuralModelTileType::OneByOne;
 		BatchSizeOverride = 1;
+		TileOverlap = FIntPoint(0, 0);
+		TileOverlapResolveType = ETileOverlapResolveType::Ignore;
 	}
 
 	void Invalidate()
@@ -143,9 +170,11 @@ namespace NeuralProfile
 
 		virtual void UpdateTileType(int32 AllocationId, ENeuralModelTileType ModelTileSize) = 0;
 		virtual bool UpdateBatchSize(int32 AllocationId, int32 BatchSize) = 0;
-	
+		virtual void UpdateTileOverlap(int32 AllocationId, FIntPoint TileOverlap) = 0;
+		virtual void UpdateTileOverlapResolveType(int32 AllocationId, ETileOverlapResolveType TileOverlapResolveType) = 0;
+
 		virtual FIntVector4 GetInputDimension(UObject* NNEModelData, FString RuntimeName) = 0;
-		virtual FIntVector4 GetOutputDimension(UObject* NNEModelData, FString RuntimeName) = 0;
+		virtual FIntVector4 GetOutputDimension(UObject* NNEModelData, FString RuntimeName) = 0;		
 	};
 }
 
diff --git a/Engine/Source/Runtime/Engine/Private/Rendering/NeuralProfile.cpp b/Engine/Source/Runtime/Engine/Private/Rendering/NeuralProfile.cpp
--- a/Engine/Source/Runtime/Engine/Private/Rendering/NeuralProfile.cpp
+++ b/Engine/Source/Runtime/Engine/Private/Rendering/NeuralProfile.cpp
@@ -20,6 +20,16 @@ DEFINE_LOG_CATEGORY_STATIC(LogNeuralProfile, Log, All);
 
 TUniquePtr<NeuralProfile::INeuralProfileManager> GNeuralProfileManager;
 
+static const TCHAR* GetNeuralProfileRuntimeName(ENeuralProfileRuntimeType NeuralProfileRuntimeType)
+{
+	static const TCHAR* const kRuntimeNames[] = {
+		TEXT("NNERuntimeRDGDml"),
+	};
+
+	static_assert(UE_ARRAY_COUNT(kRuntimeNames) == int32(ENeuralProfileRuntimeType::MAX), "Fix me");
+	return kRuntimeNames[int32(NeuralProfileRuntimeType)];
+}
+
 ///////////////////////////////////////////////////////////////////////////////////////////////////
 // FNeuralProfileModelTextureManager
 
@@ -102,30 +112,39 @@ public:
 
 			FNeuralProfileStruct& CurrentSettings = NeuralProfileEntries[AllocationId].Settings;
 			const bool bModelIsTheSame =
-				CurrentSettings.RuntimeName == Settings.RuntimeName &&
+				CurrentSettings.RuntimeType == Settings.RuntimeType &&
 				CurrentSettings.NNEModelData == Settings.NNEModelData;
 
 			const bool bNeedToUpdateModel = !bModelIsTheSame;
 			
+			// Update model
 			if (bNeedToUpdateModel)
 			{
-				// any changes to the setting require an update of the NNE model.
-				if (Settings.NNEModelData.Get())
-				{
-					GNeuralProfileManager->UpdateModel(AllocationId, Settings.NNEModelData.Get(), Settings.RuntimeName);
-				}
+				// Any changes to the setting require an update of the NNE model.
+				GNeuralProfileManager->UpdateModel(AllocationId, Settings.NNEModelData.Get(), GetNeuralProfileRuntimeName(Settings.RuntimeType));
 			}
 
-			//update model tile type
-			if (CurrentSettings.TileSize != Settings.TileSize)
+			// Update tile setting
+			if (CurrentSettings.TileSize != Settings.TileSize || bNeedToUpdateModel)
 			{
 				GNeuralProfileManager->UpdateTileType(AllocationId, Settings.TileSize);
 			}
+
+			if (CurrentSettings.TileOverlap != Settings.TileOverlap || bNeedToUpdateModel)
+			{
+				GNeuralProfileManager->UpdateTileOverlap(AllocationId, Settings.TileOverlap);
+			}
 			
+			if (CurrentSettings.TileOverlapResolveType != Settings.TileOverlapResolveType || bNeedToUpdateModel)
+			{
+				GNeuralProfileManager->UpdateTileOverlapResolveType(AllocationId, Settings.TileOverlapResolveType);
+			}
+
+			// Update batch dimension
 			const bool bBufferSizeNeedsUpdate =
 				CurrentSettings.BatchSizeOverride != Settings.BatchSizeOverride;
 			int32 BatchSizeOverride = NeuralProfileEntries[AllocationId].Settings.BatchSizeOverride;
-			if (bBufferSizeNeedsUpdate)
+			if (bBufferSizeNeedsUpdate/* || bNeedToUpdateModel*/)
 			{
 				if (GNeuralProfileManager->UpdateBatchSize(AllocationId, Settings.BatchSizeOverride))
 				{
@@ -247,11 +266,11 @@ void UNeuralProfile::PostEditChangeProperty(struct FPropertyChangedEvent& Proper
 			{
 				Settings.InputDimension = GNeuralProfileManager->GetInputDimension(
 					Settings.NNEModelData.Get(),
-					Settings.RuntimeName);
+					GetNeuralProfileRuntimeName(Settings.RuntimeType));
 
 				Settings.OutputDimension = GNeuralProfileManager->GetOutputDimension(
 					Settings.NNEModelData.Get(),
-					Settings.RuntimeName);
+					GetNeuralProfileRuntimeName(Settings.RuntimeType));
 			});
 	}
 }
