Fix NeuralPostProcessingCS.cpp not submitted.

diff --git a/Engine/Plugins/Experimental/NeuralRendering/Source/NeuralPostProcessing/Private/NeuralPostProcessingCS.cpp b/Engine/Plugins/Experimental/NeuralRendering/Source/NeuralPostProcessing/Private/NeuralPostProcessingCS.cpp
--- a/Engine/Plugins/Experimental/NeuralRendering/Source/NeuralPostProcessing/Private/NeuralPostProcessingCS.cpp
+++ b/Engine/Plugins/Experimental/NeuralRendering/Source/NeuralPostProcessing/Private/NeuralPostProcessingCS.cpp
@@ -3,39 +3,54 @@
 
 #include "NeuralPostProcessingCS.h"
 
-
-FNueralPostProcessInput GetNeuralPostProcessInput(FRDGTextureRef Texture, const FScreenPassTextureViewportParameters& ViewportParameters)
-{
-	FNueralPostProcessInput Input;
-	Input.Texture = Texture;
-	Input.Viewport = ViewportParameters;
-	return Input;
-}
-
-
-void FNeuralPostProcessingBuildIndirectDispatchArgsCS::ModifyCompilationEnvironment(const FGlobalShaderPermutationParameters& Parameters, FShaderCompilerEnvironment& OutEnvironment)
-{
-	FGlobalShader::ModifyCompilationEnvironment(Parameters, OutEnvironment);
-	OutEnvironment.SetDefine(TEXT("THREAD_GROUP_SIZE"), NEURAL_POST_PROCESSING_THREAD_GROUP_SIZE);
-}
-
-bool FNeuralPostProcessingBuildIndirectDispatchArgsCS::ShouldCompilePermutation(const FGlobalShaderPermutationParameters& Parameters)
+namespace NeuralPostProcessng
 {
-	return true;
-}
-
-void FNeuralPostProcessingPrepareInputCS::ModifyCompilationEnvironment(const FGlobalShaderPermutationParameters& InParameters, FShaderCompilerEnvironment& OutEnvironment)
-{
-	FGlobalShader::ModifyCompilationEnvironment(InParameters, OutEnvironment);
-	OutEnvironment.SetDefine(TEXT("THREAD_GROUP_SIZE"), NEURAL_POST_PROCESSING_THREAD_GROUP_SIZE);
-}
-
-void FNeuralPostProcessingProcessOutputPS::ModifyCompilationEnvironment(const FGlobalShaderPermutationParameters& InParameters, FShaderCompilerEnvironment& OutEnvironment)
-{
-	FGlobalShader::ModifyCompilationEnvironment(InParameters, OutEnvironment);
-	OutEnvironment.SetDefine(TEXT("THREAD_GROUP_SIZE"), NEURAL_POST_PROCESSING_THREAD_GROUP_SIZE);
-}
-
-IMPLEMENT_GLOBAL_SHADER(FNeuralPostProcessingBuildIndirectDispatchArgsCS, "/NeuralRendering/NeuralPostProcessing.usf", "BuildIndirectDispatchArgsCS", SF_Compute);
-IMPLEMENT_GLOBAL_SHADER(FNeuralPostProcessingPrepareInputCS, "/NeuralRendering/NeuralPostProcessing.usf", "PrepareInput", SF_Compute);
-IMPLEMENT_GLOBAL_SHADER(FNeuralPostProcessingProcessOutputPS, "/NeuralRendering/NeuralPostProcessing.usf", "ProcessOutput", SF_Pixel);
\ No newline at end of file
+	FNueralPostProcessInput GetNeuralPostProcessInput(FRDGTextureRef Texture, const FScreenPassTextureViewportParameters& ViewportParameters)
+	{
+		FNueralPostProcessInput Input;
+		Input.Texture = Texture;
+		Input.Viewport = ViewportParameters;
+		return Input;
+	}
+
+	void FNeuralPostProcessingBuildIndirectDispatchArgsCS::ModifyCompilationEnvironment(const FGlobalShaderPermutationParameters& Parameters, FShaderCompilerEnvironment& OutEnvironment)
+	{
+		FGlobalShader::ModifyCompilationEnvironment(Parameters, OutEnvironment);
+		OutEnvironment.SetDefine(TEXT("THREAD_GROUP_SIZE"), NEURAL_POST_PROCESSING_THREAD_GROUP_SIZE);
+	}
+
+	bool FNeuralPostProcessingBuildIndirectDispatchArgsCS::ShouldCompilePermutation(const FGlobalShaderPermutationParameters& Parameters)
+	{
+		return true;
+	}
+
+	void FDownScaleTextureCS::ModifyCompilationEnvironment(const FGlobalShaderPermutationParameters& InParameters, FShaderCompilerEnvironment& OutEnvironment)
+	{
+		FGlobalShader::ModifyCompilationEnvironment(InParameters, OutEnvironment);
+		OutEnvironment.SetDefine(TEXT("THREAD_GROUP_SIZE"), NEURAL_POST_PROCESSING_THREAD_GROUP_SIZE);
+	}
+
+	void FDownScaleTexture::ModifyCompilationEnvironment(const FGlobalShaderPermutationParameters& InParameters, FShaderCompilerEnvironment& OutEnvironment)
+	{
+		FGlobalShader::ModifyCompilationEnvironment(InParameters, OutEnvironment);
+	}
+
+	void FUpscaleTexture::ModifyCompilationEnvironment(const FGlobalShaderPermutationParameters& InParameters, FShaderCompilerEnvironment& OutEnvironment)
+	{
+		FGlobalShader::ModifyCompilationEnvironment(InParameters, OutEnvironment);
+		OutEnvironment.SetDefine(TEXT("THREAD_GROUP_SIZE"), NEURAL_POST_PROCESSING_THREAD_GROUP_SIZE);
+	}
+
+	void FCopyBetweenTextureAndOverlappedTileBufferCS::ModifyCompilationEnvironment(const FGlobalShaderPermutationParameters& InParameters, FShaderCompilerEnvironment& OutEnvironment)
+	{
+		FGlobalShader::ModifyCompilationEnvironment(InParameters, OutEnvironment);
+		OutEnvironment.SetDefine(TEXT("THREAD_GROUP_SIZE"), NEURAL_POST_PROCESSING_THREAD_GROUP_SIZE);
+	}
+
+	IMPLEMENT_GLOBAL_SHADER(FNeuralPostProcessingBuildIndirectDispatchArgsCS, "/NeuralRendering/NeuralPostProcessing.usf", "BuildIndirectDispatchArgsCS", SF_Compute);
+
+	IMPLEMENT_GLOBAL_SHADER(FDownScaleTextureCS, "/NeuralRendering/NeuralPostProcessing.usf", "DownScaleTextureCS", SF_Compute);
+	IMPLEMENT_GLOBAL_SHADER(FDownScaleTexture, "/NeuralRendering/NeuralPostProcessing.usf", "DownScaleTexture", SF_Pixel);
+	IMPLEMENT_GLOBAL_SHADER(FUpscaleTexture, "/NeuralRendering/NeuralPostProcessing.usf", "UpscaleTexture", SF_Pixel);
+	IMPLEMENT_GLOBAL_SHADER(FCopyBetweenTextureAndOverlappedTileBufferCS, "/NeuralRendering/NeuralPostProcessing.usf", "CopyBetweenTextureAndOverlappedTileBufferCS", SF_Compute);
+}
\ No newline at end of file
