Add NNERuntimeRDGHlsl support for neural post processing, and catch more errors into log.

diff --git a/Engine/Plugins/Experimental/NeuralRendering/Source/NeuralPostProcessing/Private/NeuralPostProcessModelInstance.cpp b/Engine/Plugins/Experimental/NeuralRendering/Source/NeuralPostProcessing/Private/NeuralPostProcessModelInstance.cpp
--- a/Engine/Plugins/Experimental/NeuralRendering/Source/NeuralPostProcessing/Private/NeuralPostProcessModelInstance.cpp
+++ b/Engine/Plugins/Experimental/NeuralRendering/Source/NeuralPostProcessing/Private/NeuralPostProcessModelInstance.cpp
@@ -31,7 +31,7 @@ TSharedPtr<UE::NNE::IModelInstanceRDG> CreateNNEModelInstance(UNNEModelData* NNE
 	if (!ModelRDG.IsValid())
 	{
 #if WITH_EDITOR
-		UE_LOG(LogNeuralPostProcessing, Error, TEXT("CreateModelRDG failed for Runtime = %s"), *RuntimeName);
+		UE_LOG(LogNeuralPostProcessing, Error, TEXT("CreateModelRDG failed for Model = %s, Runtime = %s"), *NNEModelData->GetName(), *RuntimeName);
 #endif
 		return nullptr;
 	}
@@ -65,7 +65,7 @@ TSharedPtr<UE::NNE::IModelInstanceCPU> CreateNNECpuModelInstance(UNNEModelData*
 	if (!ModelCPU.IsValid())
 	{
 #if WITH_EDITOR
-		UE_LOG(LogNeuralPostProcessing, Error, TEXT("CreateModelCPU failed for Runtime = %s"), *RuntimeName);
+		UE_LOG(LogNeuralPostProcessing, Error, TEXT("CreateModelCPU failed for Model = %s, Runtime = %s"), *NNEModelData->GetName(), *RuntimeName);
 #endif
 		return nullptr;
 	}
@@ -295,7 +295,14 @@ void UNeuralPostProcessModelInstance::CreateDefaultNNEModel(UNNEModelData* NNEMo
 		// All dynamic dimensions are set to 1 by default.
 		ResolvedInputTensorShape = UE::NNE::FTensorShape::MakeFromSymbolic(InputShape);
 
-		ModelInstanceRDG->SetInputTensorShapes({ ResolvedInputTensorShape });
+		if (ModelInstanceRDG->SetInputTensorShapes({ ResolvedInputTensorShape }) != 0)
+		{
+			ModelInstanceRDG.Reset();
+#if WITH_EDITOR
+			UE_LOG(LogNeuralPostProcessing, Warning, TEXT("SetInputTensorShapes Failed for NNE RDG Model = %s, Runtime = %s"), *NNEModelData->GetName(), *RuntimeName);
+#endif
+		}
+		
 	}
 	else
 	{
diff --git a/Engine/Plugins/Experimental/NeuralRendering/Source/NeuralPostProcessing/Private/NeuralPostProcessing.cpp b/Engine/Plugins/Experimental/NeuralRendering/Source/NeuralPostProcessing/Private/NeuralPostProcessing.cpp
--- a/Engine/Plugins/Experimental/NeuralRendering/Source/NeuralPostProcessing/Private/NeuralPostProcessing.cpp
+++ b/Engine/Plugins/Experimental/NeuralRendering/Source/NeuralPostProcessing/Private/NeuralPostProcessing.cpp
@@ -29,9 +29,10 @@ namespace
 		TEXT(" 1: enabled (default)"),
 		ECVF_RenderThreadSafe | ECVF_Scalability);
 
-	TAutoConsoleVariable<int32> CVarNeuralPostProcessTileOverlap(
+	int32 GNeuralPostProcessTileOverlap = -1;
+	FAutoConsoleVariableRef CVarNeuralPostProcessTileOverlap(
 		TEXT("r.Neuralpostprocess.TileOverlap"),
-		-1,
+		GNeuralPostProcessTileOverlap,
 		TEXT(" <0: Use the overlap from the profile\n"),
 		ECVF_RenderThreadSafe | ECVF_Scalability);
 
@@ -61,7 +62,7 @@ namespace
 static FIntPoint ClampOverlap(FIntPoint TileOverlap,FIntPoint BufferSizeXY)
 {
 	// Use TileOverlap from CVar if larger or equal than zero.
-	int32 OverlapFromCVar = CVarNeuralPostProcessTileOverlap.GetValueOnRenderThread();
+	int32 OverlapFromCVar = GNeuralPostProcessTileOverlap;
 	bool ShouldOverrideOverlapWithCVar = OverlapFromCVar >= 0;
 	if (ShouldOverrideOverlapWithCVar)
 	{
diff --git a/Engine/Source/Runtime/Engine/Classes/Engine/NeuralProfile.h b/Engine/Source/Runtime/Engine/Classes/Engine/NeuralProfile.h
--- a/Engine/Source/Runtime/Engine/Classes/Engine/NeuralProfile.h
+++ b/Engine/Source/Runtime/Engine/Classes/Engine/NeuralProfile.h
@@ -52,6 +52,10 @@ UENUM(BlueprintType)
 enum class ENeuralProfileRuntimeType : uint8
 {
 	NNERuntimeRDGDml UMETA(DisplayName = "NNERuntimeRDGDml"),
+
+	/** Does not have full operator support*/
+	NNERuntimeRDGHlsl UMETA(DisplayName = "NNERuntimeRDGHlsl"),
+
 	MAX				 UMETA(Hidden)
 };
 
diff --git a/Engine/Source/Runtime/Engine/Private/Rendering/NeuralProfile.cpp b/Engine/Source/Runtime/Engine/Private/Rendering/NeuralProfile.cpp
--- a/Engine/Source/Runtime/Engine/Private/Rendering/NeuralProfile.cpp
+++ b/Engine/Source/Runtime/Engine/Private/Rendering/NeuralProfile.cpp
@@ -24,6 +24,7 @@ static const TCHAR* GetNeuralProfileRuntimeName(ENeuralProfileRuntimeType Neural
 {
 	static const TCHAR* const kRuntimeNames[] = {
 		TEXT("NNERuntimeRDGDml"),
+		TEXT("NNERuntimeRDGHlsl")
 	};
 
 	static_assert(UE_ARRAY_COUNT(kRuntimeNames) == int32(ENeuralProfileRuntimeType::MAX), "Fix me");
