Add missing copyright info in the shader file for neural post processing.

diff --git a/Engine/Plugins/Experimental/NeuralRendering/Source/NeuralPostProcessing/Shaders/NeuralPostProcessing.usf b/Engine/Plugins/Experimental/NeuralRendering/Source/NeuralPostProcessing/Shaders/NeuralPostProcessing.usf
--- a/Engine/Plugins/Experimental/NeuralRendering/Source/NeuralPostProcessing/Shaders/NeuralPostProcessing.usf
+++ b/Engine/Plugins/Experimental/NeuralRendering/Source/NeuralPostProcessing/Shaders/NeuralPostProcessing.usf
@@ -1,4 +1,4 @@
-
+// Copyright Epic Games, Inc. All Rights Reserved.
 
 #include "/Engine/Public/Platform.ush"
 #include "/Engine/Private/Common.ush"
