[NeuralPostProcessing] Remove dependancie to NNERuntimeORTCpu.

Problem:
NNERuntimeORTCpu is being removed for 5.4, NNERuntimeORT is replacing it for the editor use case.

Solution:
Remove dependancie, using only RDG runtime instead.

Test:
Create a neural postprocess material and verify preview is applying the effect

#rnx

diff --git a/Engine/Plugins/Experimental/NeuralRendering/NeuralRendering.uplugin b/Engine/Plugins/Experimental/NeuralRendering/NeuralRendering.uplugin
--- a/Engine/Plugins/Experimental/NeuralRendering/NeuralRendering.uplugin
+++ b/Engine/Plugins/Experimental/NeuralRendering/NeuralRendering.uplugin
@@ -32,10 +32,6 @@
 		{
 			"Name": "NNERuntimeRDG",
 			"Enabled": true
-		},
-		{
-			"Name": "NNERuntimeORTCpu",
-			"Enabled": true
 		}
 	]
 }
\ No newline at end of file
diff --git a/Engine/Plugins/Experimental/NeuralRendering/Source/NeuralPostProcessing/Private/NeuralPostProcessModelInstance.cpp b/Engine/Plugins/Experimental/NeuralRendering/Source/NeuralPostProcessing/Private/NeuralPostProcessModelInstance.cpp
--- a/Engine/Plugins/Experimental/NeuralRendering/Source/NeuralPostProcessing/Private/NeuralPostProcessModelInstance.cpp
+++ b/Engine/Plugins/Experimental/NeuralRendering/Source/NeuralPostProcessing/Private/NeuralPostProcessModelInstance.cpp
@@ -39,40 +39,6 @@ TSharedPtr<UE::NNE::IModelInstanceRDG> CreateNNEModelInstance(UNNEModelData* NNE
 	return ModelRDG->CreateModelInstanceRDG();
 }
 
-TSharedPtr<UE::NNE::IModelInstanceCPU> CreateNNECpuModelInstance(UNNEModelData* NNEModelData)
-{
-	check(NNEModelData);
-	const FString RuntimeName = "NNERuntimeORTCpu";
-	TWeakInterfacePtr<INNERuntime> Runtime = UE::NNE::GetRuntime<INNERuntime>(RuntimeName);
-	TWeakInterfacePtr<INNERuntimeCPU> RuntimeCPU = UE::NNE::GetRuntime<INNERuntimeCPU>(RuntimeName);
-	if (!Runtime.IsValid())
-	{
-#if WITH_EDITOR
-		UE_LOG(LogNeuralPostProcessing, Error, TEXT("Can't get %s runtime."), *RuntimeName);
-#endif
-		return nullptr;
-	}
-
-	if (!RuntimeCPU.IsValid())
-	{
-#if WITH_EDITOR
-		UE_LOG(LogNeuralPostProcessing, Error, TEXT("No CPU runtime '%s' found"), *RuntimeName);
-#endif
-		return nullptr;
-	}
-
-	TSharedPtr<UE::NNE::IModelCPU> ModelCPU = RuntimeCPU->CreateModelCPU(NNEModelData);
-	if (!ModelCPU.IsValid())
-	{
-#if WITH_EDITOR
-		UE_LOG(LogNeuralPostProcessing, Error, TEXT("CreateModelCPU failed for Model = %s, Runtime = %s"), *NNEModelData->GetName(), *RuntimeName);
-#endif
-		return nullptr;
-	}
-
-	return ModelCPU->CreateModelInstanceCPU();
-}
-
 UNeuralPostProcessModelInstance::UNeuralPostProcessModelInstance(FObjectInitializer const& ObjectInitializer)
 	:Super(ObjectInitializer)
 {
diff --git a/Engine/Plugins/Experimental/NeuralRendering/Source/NeuralPostProcessing/Private/NeuralPostProcessModelInstance.h b/Engine/Plugins/Experimental/NeuralRendering/Source/NeuralPostProcessing/Private/NeuralPostProcessModelInstance.h
--- a/Engine/Plugins/Experimental/NeuralRendering/Source/NeuralPostProcessing/Private/NeuralPostProcessModelInstance.h
+++ b/Engine/Plugins/Experimental/NeuralRendering/Source/NeuralPostProcessing/Private/NeuralPostProcessModelInstance.h
@@ -5,7 +5,6 @@
 #include "NNE.h"
 #include "NNETypes.h"
 #include "NNERuntimeRDG.h"
-#include "NNERuntimeCPU.h"
 #include "NNEModelData.h"
 
 #include "NeuralPostProcessModelInstance.generated.h"
@@ -16,7 +15,6 @@
 
 //Util functions
 TSharedPtr<UE::NNE::IModelInstanceRDG> CreateNNEModelInstance(UNNEModelData* NNEModelData, FString RuntimeName);
-TSharedPtr<UE::NNE::IModelInstanceCPU> CreateNNECpuModelInstance(UNNEModelData* NNEModelData);
 
 enum class ENeuralModelTileType : uint8;
 enum class ETileOverlapResolveType : uint8;
diff --git a/Engine/Plugins/Experimental/NeuralRendering/Source/NeuralPostProcessing/Private/NeuralPostProcessing.cpp b/Engine/Plugins/Experimental/NeuralRendering/Source/NeuralPostProcessing/Private/NeuralPostProcessing.cpp
--- a/Engine/Plugins/Experimental/NeuralRendering/Source/NeuralPostProcessing/Private/NeuralPostProcessing.cpp
+++ b/Engine/Plugins/Experimental/NeuralRendering/Source/NeuralPostProcessing/Private/NeuralPostProcessing.cpp
@@ -217,7 +217,7 @@ public:
 		{
 			// Need to create the ModelInstance in order to get the dimension in case the Model is not created
 			// in the rendering thread.
-			TSharedPtr<UE::NNE::IModelInstanceCPU> ModelInstance = CreateNNECpuModelInstance(ModelData);
+			TSharedPtr<UE::NNE::IModelInstanceRDG> ModelInstance = CreateNNEModelInstance(ModelData, RuntimeName);
 
 			if (ModelInstance)
 			{
@@ -246,7 +246,7 @@ public:
 
 		if (UNNEModelData* ModelData = Cast<UNNEModelData>(NNEModelData))
 		{
-			TSharedPtr<UE::NNE::IModelInstanceCPU> ModelInstance = CreateNNECpuModelInstance(ModelData);
+			TSharedPtr<UE::NNE::IModelInstanceRDG> ModelInstance = CreateNNEModelInstance(ModelData, RuntimeName);
 			
 			if (ModelInstance)
 			{
