Fix shader compile error of not specifying CFLAG_AllowTypedUAVLoads with dx11 if neural rendering plugin is enabled.

diff --git a/Engine/Plugins/Experimental/NeuralRendering/Source/NeuralPostProcessing/Private/NeuralPostProcessingCS.cpp b/Engine/Plugins/Experimental/NeuralRendering/Source/NeuralPostProcessing/Private/NeuralPostProcessingCS.cpp
--- a/Engine/Plugins/Experimental/NeuralRendering/Source/NeuralPostProcessing/Private/NeuralPostProcessingCS.cpp
+++ b/Engine/Plugins/Experimental/NeuralRendering/Source/NeuralPostProcessing/Private/NeuralPostProcessingCS.cpp
@@ -2,6 +2,7 @@
 
 
 #include "NeuralPostProcessingCS.h"
+#include "ShaderCompilerCore.h"
 
 namespace NeuralPostProcessng
 {
@@ -45,6 +46,7 @@ namespace NeuralPostProcessng
 	{
 		FGlobalShader::ModifyCompilationEnvironment(InParameters, OutEnvironment);
 		OutEnvironment.SetDefine(TEXT("THREAD_GROUP_SIZE"), NEURAL_POST_PROCESSING_THREAD_GROUP_SIZE);
+		OutEnvironment.CompilerFlags.Add(CFLAG_AllowTypedUAVLoads);
 	}
 
 	IMPLEMENT_GLOBAL_SHADER(FNeuralPostProcessingBuildIndirectDispatchArgsCS, "/NeuralRendering/NeuralPostProcessing.usf", "BuildIndirectDispatchArgsCS", SF_Compute);
