Fixed usage of DECLARE_GPU_STAT (when RHI_NEW_GPU_PROFILER=1).

#rnx

diff --git a/Engine/Plugins/Experimental/NeuralRendering/Source/NeuralPostProcessing/Private/NeuralPostProcessing.cpp b/Engine/Plugins/Experimental/NeuralRendering/Source/NeuralPostProcessing/Private/NeuralPostProcessing.cpp
--- a/Engine/Plugins/Experimental/NeuralRendering/Source/NeuralPostProcessing/Private/NeuralPostProcessing.cpp
+++ b/Engine/Plugins/Experimental/NeuralRendering/Source/NeuralPostProcessing/Private/NeuralPostProcessing.cpp
@@ -56,7 +56,7 @@ namespace
 		TEXT(" Adjust the intensity of the overlap for visualization\n"),
 		ECVF_RenderThreadSafe | ECVF_Scalability);
 
-	DECLARE_GPU_STAT(NeuralPostProcessing)
+	DECLARE_GPU_STAT(NeuralPostProcessing);
 }
 
 static FIntPoint ClampOverlap(FIntPoint TileOverlap,FIntPoint BufferSizeXY)
