Fix a GPU crash when assigning a neural profile to a post process material in debug build. Also stop unneccessary model creation when the RDG model cannot be created.

#tests Test done by followng the steps in debug build, no assert

diff --git a/Engine/Plugins/Experimental/NeuralRendering/Source/NeuralPostProcessing/Private/NeuralPostProcessModelInstance.cpp b/Engine/Plugins/Experimental/NeuralRendering/Source/NeuralPostProcessing/Private/NeuralPostProcessModelInstance.cpp
--- a/Engine/Plugins/Experimental/NeuralRendering/Source/NeuralPostProcessing/Private/NeuralPostProcessModelInstance.cpp
+++ b/Engine/Plugins/Experimental/NeuralRendering/Source/NeuralPostProcessing/Private/NeuralPostProcessModelInstance.cpp
@@ -26,6 +26,14 @@ TSharedPtr<UE::NNE::IModelInstanceRDG> CreateNNEModelInstance(UNNEModelData* NNE
 #endif
 		return nullptr;
 	}
+	
+	if (RuntimeRDG->CanCreateModelRDG(NNEModelData) != INNERuntimeRDG::ECanCreateModelRDGStatus::Ok)
+	{
+#if WITH_EDITOR
+		UE_LOG(LogNeuralPostProcessing, Error, TEXT("Cannot create model RDG for Model = %s, Runtime = %s"), *NNEModelData->GetName(), *RuntimeName);
+#endif
+		return nullptr;
+	}
 
 	TSharedPtr<UE::NNE::IModelRDG> ModelRDG = RuntimeRDG->CreateModelRDG(NNEModelData);
 	if (!ModelRDG.IsValid())
diff --git a/Engine/Plugins/Experimental/NeuralRendering/Source/NeuralPostProcessing/Private/NeuralPostProcessing.cpp b/Engine/Plugins/Experimental/NeuralRendering/Source/NeuralPostProcessing/Private/NeuralPostProcessing.cpp
--- a/Engine/Plugins/Experimental/NeuralRendering/Source/NeuralPostProcessing/Private/NeuralPostProcessing.cpp
+++ b/Engine/Plugins/Experimental/NeuralRendering/Source/NeuralPostProcessing/Private/NeuralPostProcessing.cpp
@@ -157,7 +157,11 @@ public:
 			return;
 		}
 
-		UNeuralPostProcessModelInstance* ModelInstance = NewObject<UNeuralPostProcessModelInstance>();
+		UNeuralPostProcessModelInstance* ModelInstance = nullptr;
+		{
+			FGCScopeGuard GCScopeGuard;
+			ModelInstance = NewObject<UNeuralPostProcessModelInstance>();
+		}
 		ModelInstance->Update(RawNNEModelData, RuntimeName);
 
 		FNeuralPostProcessModelInstanceManager::Get()->SetModelInstance(AllocationId, ModelInstance->IsValid()? ModelInstance : nullptr);
