[Neural post processing] Fix cooking fail in editor.

#tests Test with editor on windows

diff --git a/Engine/Plugins/Experimental/NeuralRendering/Source/NeuralPostProcessing/Private/NeuralPostProcessModelInstance.cpp b/Engine/Plugins/Experimental/NeuralRendering/Source/NeuralPostProcessing/Private/NeuralPostProcessModelInstance.cpp
--- a/Engine/Plugins/Experimental/NeuralRendering/Source/NeuralPostProcessing/Private/NeuralPostProcessModelInstance.cpp
+++ b/Engine/Plugins/Experimental/NeuralRendering/Source/NeuralPostProcessing/Private/NeuralPostProcessModelInstance.cpp
@@ -22,7 +22,7 @@ TSharedPtr<UE::NNE::IModelInstanceRDG> CreateNNEModelInstance(UNNEModelData* NNE
 	if (!RuntimeRDG.IsValid())
 	{
 #if WITH_EDITOR
-		UE_LOG(LogNeuralPostProcessing, Error, TEXT("No RDG runtime '%s' found"), *RuntimeName);
+		UE_LOG(LogNeuralPostProcessing, Log, TEXT("No RDG runtime '%s' found"), *RuntimeName);
 #endif
 		return nullptr;
 	}
@@ -281,7 +281,7 @@ void UNeuralPostProcessModelInstance::CreateDefaultNNEModel(UNNEModelData* NNEMo
 	else
 	{
 #if WITH_EDITOR
-		UE_LOG(LogNeuralPostProcessing, Error, TEXT("Failed to create NNE RDG Model."));
+		UE_LOG(LogNeuralPostProcessing, Log, TEXT("Failed to create NNE RDG Model."));
 #endif
 	}
 }
