[Neural post processing] Fix typo

diff --git a/Engine/Plugins/Experimental/NeuralRendering/Source/NeuralPostProcessing/Private/NeuralPostProcessModelInstance.cpp b/Engine/Plugins/Experimental/NeuralRendering/Source/NeuralPostProcessing/Private/NeuralPostProcessModelInstance.cpp
--- a/Engine/Plugins/Experimental/NeuralRendering/Source/NeuralPostProcessing/Private/NeuralPostProcessModelInstance.cpp
+++ b/Engine/Plugins/Experimental/NeuralRendering/Source/NeuralPostProcessing/Private/NeuralPostProcessModelInstance.cpp
@@ -181,7 +181,7 @@ bool UNeuralPostProcessModelInstance::ModifyInputShape(int Dim, int Size)
 	if (ResolvedOutputTensorShapes.Num() < 1)
 	{
 #if WITH_EDITOR
-		UE_LOG(LogNeuralPostProcessing, Error, TEXT("Cannot set dimension %d to %d due to NNE eror. Revert back to %d"), Dim, Size, CurrentResolvedInputTensorShape.GetData()[Dim]);
+		UE_LOG(LogNeuralPostProcessing, Error, TEXT("Cannot set dimension %d to %d due to NNE error. Revert back to %d"), Dim, Size, CurrentResolvedInputTensorShape.GetData()[Dim]);
 #endif
 		ModelInstanceRDG->SetInputTensorShapes({ CurrentResolvedInputTensorShape });
 		return false;
